Introduced a new plugin to host the Editor Framework's new data storage.

This changelist introduces a new plugin that provides an implementation for the Editor Framework's new data storage. The data storage can be seen as a central database-like structure that eventually will become the single source of truth for all data used by the editor. Processors can be added that (periodically) run over query results allowing arbitrary data to be displayed that matches the data. This in turn can be used to create reusable UI elements and other functionality. Users can add their own columns with data for custom and novel functionality as well.

This first addition introduces:
- The plugin itself.
- Two new interfaces to interact with the data storage. To be expanded upon in the future.
- Updates to the Typed Elements API to access the new API.
- Automatically start using the new data storage if the plugin is enabled.
- Register any actors in the editor and unregister when they get garbage collected or explicitly destroyed.

Other than being able to see the number of actors in the MASS Debugger, there's no no or replacement functionality introduced.

diff --git a/Engine/Plugins/TypedElementsDataStorage/Source/TypedElementsDataStorage/Private/TypedElementDatabase.cpp b/Engine/Plugins/TypedElementsDataStorage/Source/TypedElementsDataStorage/Private/TypedElementDatabase.cpp
new file mode 100644
--- /dev/null
+++ b/Engine/Plugins/TypedElementsDataStorage/Source/TypedElementsDataStorage/Private/TypedElementDatabase.cpp
@@ -0,0 +1,185 @@
+// Copyright Epic Games, Inc. All Rights Reserved.
+
+#include "TypedElementDatabase.h"
+
+#include "Editor.h"
+#include "Engine/World.h"
+#include "MassEntitySubsystem.h"
+
+void UTypedElementDatabase::Initialize()
+{
+	ActiveEditorWorld = GEditor->GetEditorWorldContext().World();
+	check(ActiveEditorWorld);
+	UMassEntitySubsystem* Mass = ActiveEditorWorld->GetSubsystem<UMassEntitySubsystem>();
+	check(Mass);
+	ActiveEditorEntityManager = Mass->GetMutableEntityManager().AsShared();
+	
+	TickCallHandle = FWorldDelegates::OnWorldTickStart.AddLambda([this](UWorld* World, ELevelTick, float)
+		{
+			if (World == ActiveEditorWorld)
+			{
+				OnUpdateDelegate.Broadcast();
+			}
+		});
+
+	GEditor->OnWorldAdded().AddLambda([this](UWorld* NewWorld)
+		{
+			if (GEditor->GetEditorWorldContext().World() == NewWorld)
+			{
+				ActiveEditorWorld = NewWorld;
+				UMassEntitySubsystem* Mass = NewWorld->GetSubsystem<UMassEntitySubsystem>();
+				check(Mass);
+				ActiveEditorEntityManager = Mass->GetMutableEntityManager().AsShared();
+				
+				OnCreationDelegate.Broadcast();
+			}
+		});
+	GEditor->OnWorldDestroyed().AddLambda([this](UWorld* DestroyedWorld)
+		{
+			if (DestroyedWorld == ActiveEditorWorld)
+			{
+				Reset();
+			}
+		});
+}
+
+void UTypedElementDatabase::Deinitialize()
+{
+	FWorldDelegates::OnWorldTickStart.Remove(TickCallHandle);
+
+	Reset();
+}
+
+TSharedPtr<FMassEntityManager> UTypedElementDatabase::GetActiveMutableEditorEntityManager()
+{
+	return ActiveEditorEntityManager;
+}
+
+TSharedPtr<const FMassEntityManager> UTypedElementDatabase::GetActiveEditorEntityManager() const
+{
+	return ActiveEditorEntityManager;
+}
+
+FTypedElementOnDataStorageCreation& UTypedElementDatabase::OnCreation()
+{
+	return OnCreationDelegate;
+}
+
+FTypedElementOnDataStorageDestruction& UTypedElementDatabase::OnDestruction()
+{
+	return OnDestructionDelegate;
+}
+
+FTypedElementOnDataStorageUpdate& UTypedElementDatabase::OnUpdate()
+{
+	return OnUpdateDelegate;
+}
+
+bool UTypedElementDatabase::IsAvailable() const
+{
+	return ActiveEditorWorld && ActiveEditorEntityManager;
+}
+
+TypedElementTableHandle UTypedElementDatabase::RegisterTable(TConstArrayView<const UScriptStruct*> ColumnList)
+{
+	return RegisterTable(ColumnList, {});
+}
+
+TypedElementTableHandle UTypedElementDatabase::RegisterTable(TConstArrayView<const UScriptStruct*> ColumnList, const FName Name)
+{
+	if (ActiveEditorEntityManager)
+	{
+		if (Name.IsValid())
+		{
+			if (!TableNameLookup.Contains(Name))
+			{
+				TypedElementTableHandle Result = Tables.Num();
+				Tables.Add(ActiveEditorEntityManager->CreateArchetype(ColumnList, Name));
+				TableNameLookup.Add(Name, Result);
+				return Result;
+			}
+		}
+		else
+		{
+			TypedElementTableHandle Result = Tables.Num();
+			Tables.Add(ActiveEditorEntityManager->CreateArchetype(ColumnList, Name));
+			return Result;
+		}
+	}
+	return TypedElementInvalidTableHandle;
+}
+
+TypedElementRowHandle UTypedElementDatabase::AddRow(TypedElementTableHandle Table)
+{
+	checkf(Table < Tables.Num(), TEXT("Attempting to add a row to a non-existing table."));
+	return ActiveEditorEntityManager ? 
+		ActiveEditorEntityManager->CreateEntity(Tables[Table]).AsNumber() :
+		TypedElementInvalidRowHandle;
+}
+
+TypedElementRowHandle UTypedElementDatabase::AddRow(FName TableName)
+{
+	TypedElementTableHandle* Table = TableNameLookup.Find(TableName);
+	return Table ? AddRow(*Table) : TypedElementInvalidRowHandle;
+}
+
+bool UTypedElementDatabase::BatchAddRow(TypedElementTableHandle Table, int32 Count, TypedElementDataStorageCreationCallbackRef OnCreated)
+{
+	OnCreated.CheckCallable();
+	checkf(Table < Tables.Num(), TEXT("Attempting to add multiple rows to a non-existing table."));
+	if (ActiveEditorEntityManager)
+	{
+		TArray<FMassEntityHandle> Entities;
+		Entities.Reserve(Count);
+		TSharedRef<FMassEntityManager::FEntityCreationContext> Context = 
+			ActiveEditorEntityManager->BatchCreateEntities(Tables[Table], Count, Entities);
+		
+		for (FMassEntityHandle Entity : Entities)
+		{
+			OnCreated(Entity.AsNumber());
+		}
+
+		return true;
+	}
+	return false;
+}
+
+bool UTypedElementDatabase::BatchAddRow(FName TableName, int32 Count, TypedElementDataStorageCreationCallbackRef OnCreated)
+{
+	TypedElementTableHandle* Table = TableNameLookup.Find(TableName);
+	return Table ? BatchAddRow(*Table, Count, OnCreated) : false;
+}
+
+void UTypedElementDatabase::RemoveRow(TypedElementRowHandle Row)
+{
+	if (ActiveEditorEntityManager)
+	{
+		ActiveEditorEntityManager->DestroyEntity(FMassEntityHandle::FromNumber(Row));
+	}
+}
+
+void* UTypedElementDatabase::GetColumnData(TypedElementRowHandle Row, const UScriptStruct* ColumnType)
+{
+	check(ColumnType);
+
+	FMassEntityHandle Entity = FMassEntityHandle::FromNumber(Row);
+	return ActiveEditorEntityManager && ActiveEditorEntityManager->IsEntityValid(Entity) ?
+		ActiveEditorEntityManager->GetFragmentDataStruct(FMassEntityHandle::FromNumber(Row), ColumnType).GetMutableMemory() :
+		nullptr;
+}
+
+void* UTypedElementDatabase::GetExternalSystemAddress(UClass* Target)
+{
+	return ActiveEditorWorld ? ActiveEditorWorld->GetSubsystemBase(Target) : nullptr;
+}
+
+void UTypedElementDatabase::Reset()
+{
+	OnDestructionDelegate.Broadcast();
+
+	Tables.Reset();
+	TableNameLookup.Reset();
+
+	ActiveEditorWorld = nullptr;
+	ActiveEditorEntityManager.Reset();
+}
\ No newline at end of file
diff --git a/Engine/Plugins/TypedElementsDataStorage/Source/TypedElementsDataStorage/Private/TypedElementDatabase.h b/Engine/Plugins/TypedElementsDataStorage/Source/TypedElementsDataStorage/Private/TypedElementDatabase.h
new file mode 100644
--- /dev/null
+++ b/Engine/Plugins/TypedElementsDataStorage/Source/TypedElementsDataStorage/Private/TypedElementDatabase.h
@@ -0,0 +1,62 @@
+// Copyright Epic Games, Inc. All Rights Reserved.
+
+#pragma once
+
+#include "CoreMinimal.h"
+#include "Templates/SharedPointer.h"
+#include "Elements/Interfaces/TypedElementDataStorageInterface.h"
+#include "MassArchetypeTypes.h"
+
+#include "TypedElementDatabase.generated.h"
+
+struct FMassEntityManager;
+class UWorld;
+
+UCLASS()
+class TYPEDELEMENTSDATASTORAGE_API UTypedElementDatabase 
+	: public UObject
+	, public ITypedElementDataStorageInterface
+{
+	GENERATED_BODY()
+public:
+	~UTypedElementDatabase() override = default;
+
+	void Initialize();
+	void Deinitialize();
+
+	TSharedPtr<FMassEntityManager> GetActiveMutableEditorEntityManager();
+	TSharedPtr<const FMassEntityManager> GetActiveEditorEntityManager() const;
+
+	FTypedElementOnDataStorageCreation& OnCreation() override;
+	FTypedElementOnDataStorageDestruction& OnDestruction() override;
+	FTypedElementOnDataStorageUpdate& OnUpdate() override;
+	bool IsAvailable() const override;
+
+	TypedElementTableHandle RegisterTable(TConstArrayView<const UScriptStruct*> ColumnList) override;
+	TypedElementTableHandle RegisterTable(TConstArrayView<const UScriptStruct*> ColumnList, const FName Name) override;
+
+	TypedElementRowHandle AddRow(TypedElementTableHandle Table) override;
+	TypedElementRowHandle AddRow(FName TableName) override;
+	bool BatchAddRow(TypedElementTableHandle Table, int32 Count, TypedElementDataStorageCreationCallbackRef OnCreated) override;
+	bool BatchAddRow(FName TableName, int32 Count, TypedElementDataStorageCreationCallbackRef OnCreated) override;
+	void RemoveRow(TypedElementRowHandle Row) override;
+
+	void* GetColumnData(TypedElementRowHandle Row, const UScriptStruct* ColumnType) override;
+
+	void* GetExternalSystemAddress(UClass* Target) override;
+
+private:
+	void Reset();
+	
+	TArray<FMassArchetypeHandle> Tables;
+	TMap<FName, TypedElementTableHandle> TableNameLookup;
+
+	FTypedElementOnDataStorageCreation OnCreationDelegate;
+	FTypedElementOnDataStorageDestruction OnDestructionDelegate;
+	FTypedElementOnDataStorageUpdate OnUpdateDelegate;
+
+	FDelegateHandle TickCallHandle;
+
+	TSharedPtr<FMassEntityManager> ActiveEditorEntityManager;
+	UWorld* ActiveEditorWorld{ nullptr };
+};
\ No newline at end of file
diff --git a/Engine/Plugins/TypedElementsDataStorage/Source/TypedElementsDataStorage/Private/TypedElementDatabaseCompatibility.cpp b/Engine/Plugins/TypedElementsDataStorage/Source/TypedElementsDataStorage/Private/TypedElementDatabaseCompatibility.cpp
new file mode 100644
--- /dev/null
+++ b/Engine/Plugins/TypedElementsDataStorage/Source/TypedElementsDataStorage/Private/TypedElementDatabaseCompatibility.cpp
@@ -0,0 +1,128 @@
+// Copyright Epic Games, Inc. All Rights Reserved.
+
+#include "TypedElementDatabaseCompatibility.h"
+
+#include "MassActorSubsystem.h"
+
+void UTypedElementDatabaseCompatibility::Initialize(ITypedElementDataStorageInterface* StorageInterface)
+{
+	checkf(StorageInterface, TEXT("Typed Element's Database compatibility manager is being initialized with an invalid storage target."));
+
+	Storage = StorageInterface;
+	Prepare();
+
+	StorageInterface->OnCreation().AddUObject(this, &UTypedElementDatabaseCompatibility::Prepare);
+	StorageInterface->OnDestruction().AddUObject(this, &UTypedElementDatabaseCompatibility::Reset);
+	StorageInterface->OnUpdate().AddUObject(this, &UTypedElementDatabaseCompatibility::AddPendingActors);
+}
+
+void UTypedElementDatabaseCompatibility::Deinitialize()
+{
+	Reset();
+}
+
+void UTypedElementDatabaseCompatibility::AddCompatibleObject(AActor* Actor)
+{
+	// Registration is delayed for two reasons:
+	//	1. Allows entity creation in a single batch rather than multiple individual additions.
+	//	2. Provides an opportunity to filter out the actors that are created within MASS itself as those will already be registered.
+	ActorsPendingRegistration.Add(Actor);
+}
+
+void UTypedElementDatabaseCompatibility::RemoveCompatibleObject(AActor* Actor)
+{
+	checkf(Storage, TEXT("Removing compatible objects is not supported before Typed Element's Database compatibility manager has been initialized."));
+
+	// If there is no actor subsystem it means that the world has been destroyed, including the MASS instance,
+	// so there's no references to clean up.
+	if (Storage && ActorSubsystem && Storage->IsAvailable() && Actor->GetWorld() == ActorSubsystem->GetWorld())
+	{
+		FMassEntityHandle Entity = ActorSubsystem->GetEntityHandleFromActor(Actor);
+		if (Entity.IsValid()) // If there's no entity it may have been deleted earlier, e.g. through an explicit delete.
+		{
+			auto ActorStore = Storage->GetColumn<FMassActorFragment>(Entity.AsNumber());
+			if (ActorStore && !ActorStore->IsOwnedByMass()) // Only remove actors that were externally created.
+			{
+				ActorSubsystem->RemoveHandleForActor(Actor);
+				Storage->RemoveRow(Entity.AsNumber());
+			}
+		}
+	}
+}
+
+void UTypedElementDatabaseCompatibility::Prepare()
+{
+	ActorSubsystem = Storage->GetExternalSystem<UMassActorSubsystem>();
+
+	CreateStandardArchetypes();
+}
+
+void UTypedElementDatabaseCompatibility::Reset()
+{
+	ActorSubsystem = nullptr;
+}
+
+void UTypedElementDatabaseCompatibility::CreateStandardArchetypes()
+{
+	StandardActorTable = Storage->RegisterTable(MakeArrayView(
+		{
+			FMassActorFragment::StaticStruct()
+		}), FName("Editor_StandardActorTable"));
+}
+
+void UTypedElementDatabaseCompatibility::AddPendingActors()
+{
+	// Delay processing until the required systems are available by not clearing the pending actor list.
+	if (ActorsPendingRegistration.Num() > 0 && Storage && Storage->IsAvailable() && ActorSubsystem)
+	{
+		UWorld* TargetWorld = ActorSubsystem->GetWorld();
+
+		// Filter out the actors that are already registered or already destroyed. 
+		// The most common case for this is actors created from within MASS.
+		TWeakObjectPtr<AActor>* ActorBegin = ActorsPendingRegistration.GetData();
+		TWeakObjectPtr<AActor>* ActorIt = ActorsPendingRegistration.GetData();
+		TWeakObjectPtr<AActor>* ActorEnd = ActorBegin + ActorsPendingRegistration.Num();
+		while (ActorIt != ActorEnd)
+		{
+			if (
+				(*ActorIt).IsValid() && // If not valid then the actor has already been deleted again.
+				(*ActorIt)->GetWorld() == TargetWorld && // If not equal than this is not an actor for the global data storage.
+				!ActorSubsystem->GetEntityHandleFromActor((*ActorIt).Get()).IsValid() // If true, this actor is likely created within MASS.
+				) 
+			{
+				++ActorIt;
+			}
+			else
+			{
+				// Don't shrink the registration array as the array will be reused with a variety of different actor counts.
+				// If memory size becomes an issue it's better to resize the array once after this loop rather than within the
+				// loop to avoid many resizes happening.
+				constexpr bool bAllowToShrink = false;
+				ActorsPendingRegistration.RemoveAtSwap(ActorIt - ActorBegin, 1, bAllowToShrink);
+				--ActorEnd;
+			}
+		}
+
+		if (ActorBegin != ActorEnd) // Check if the number of actors to add isn't zero.
+		{
+			// Add the remaining actors to the data storage.
+			ActorIt = ActorsPendingRegistration.GetData();
+			Storage->BatchAddRow(StandardActorTable, ActorsPendingRegistration.Num(), [this, &ActorIt, ActorEnd](TypedElementRowHandle Row)
+				{
+					auto ActorStore = Storage->GetColumn<FMassActorFragment>(Row);
+					if (ActorStore)
+					{
+						constexpr bool bIsOwnedByMass = false;
+						ActorStore->SetAndUpdateHandleMap(FMassEntityHandle::FromNumber(Row), ActorIt->Get(), bIsOwnedByMass);
+
+						checkf(ActorIt < ActorEnd,
+							TEXT("More (%i) entities were added than were requested (%i)."), ActorEnd - ActorIt, ActorsPendingRegistration.Num());
+						++ActorIt;
+					}
+				});
+
+			// Reset the container for next set of actors.
+			ActorsPendingRegistration.Reset();
+		}
+	}
+}
diff --git a/Engine/Plugins/TypedElementsDataStorage/Source/TypedElementsDataStorage/Private/TypedElementDatabaseCompatibility.h b/Engine/Plugins/TypedElementsDataStorage/Source/TypedElementsDataStorage/Private/TypedElementDatabaseCompatibility.h
new file mode 100644
--- /dev/null
+++ b/Engine/Plugins/TypedElementsDataStorage/Source/TypedElementsDataStorage/Private/TypedElementDatabaseCompatibility.h
@@ -0,0 +1,44 @@
+// Copyright Epic Games, Inc. All Rights Reserved.
+
+#pragma once
+
+#include "CoreMinimal.h"
+#include "Containers/Array.h"
+#include "Elements/Interfaces/TypedElementDataStorageInterface.h"
+#include "Elements/Interfaces/TypedElementDataStorageCompatibilityInterface.h"
+#include "UObject/WeakObjectPtr.h"
+
+#include "TypedElementDatabaseCompatibility.generated.h"
+
+class AActor;
+class ITypedElementDataStorageInterface;
+class UMassActorSubsystem;
+
+UCLASS()
+class TYPEDELEMENTSDATASTORAGE_API UTypedElementDatabaseCompatibility
+	: public UObject
+	, public ITypedElementDataStorageCompatibilityInterface
+{
+	GENERATED_BODY()
+public:
+	~UTypedElementDatabaseCompatibility() override = default;
+
+	void Initialize(ITypedElementDataStorageInterface* StorageInterface);
+	void Deinitialize();
+
+	void AddCompatibleObject(AActor* Actor) override;
+	void RemoveCompatibleObject(AActor* Actor) override;
+
+private:
+	void Prepare();
+	void Reset();
+	void CreateStandardArchetypes();
+	
+	void AddPendingActors();
+	
+	TArray<TWeakObjectPtr<AActor>> ActorsPendingRegistration;
+	
+	TypedElementTableHandle StandardActorTable{ TypedElementInvalidTableHandle };
+	ITypedElementDataStorageInterface* Storage{ nullptr };
+	UMassActorSubsystem* ActorSubsystem{ nullptr };
+};
\ No newline at end of file
diff --git a/Engine/Plugins/TypedElementsDataStorage/Source/TypedElementsDataStorage/Private/TypedElementsDataStorage.cpp b/Engine/Plugins/TypedElementsDataStorage/Source/TypedElementsDataStorage/Private/TypedElementsDataStorage.cpp
new file mode 100644
--- /dev/null
+++ b/Engine/Plugins/TypedElementsDataStorage/Source/TypedElementsDataStorage/Private/TypedElementsDataStorage.cpp
@@ -0,0 +1,65 @@
+// Copyright Epic Games, Inc. All Rights Reserved.
+
+#include "TypedElementsDataStorage.h"
+
+#include "Elements/Framework/TypedElementRegistry.h"
+#include "TypedElementDatabase.h"
+#include "TypedElementDatabaseCompatibility.h"
+#include "UObject/UObjectGlobals.h"
+
+#define LOCTEXT_NAMESPACE "FTypedElementsDataStorageModule"
+
+void FTypedElementsDataStorageModule::StartupModule()
+{
+	FCoreDelegates::OnAllModuleLoadingPhasesComplete.AddLambda(
+		[this]()
+		{
+			if (!bInitialized)
+			{
+				Database = NewObject<UTypedElementDatabase>();
+				Database->Initialize();
+
+				DatabaseCompatibility = NewObject<UTypedElementDatabaseCompatibility>();
+				DatabaseCompatibility->Initialize(Database.Get());
+
+				UTypedElementRegistry* Registry = UTypedElementRegistry::GetInstance();
+				check(Registry);
+				Registry->SetDataStorage(Database.Get());
+				Registry->SetDataStorageCompatibility(DatabaseCompatibility.Get());
+
+				bInitialized = true;
+			}
+		});
+}
+
+void FTypedElementsDataStorageModule::ShutdownModule()
+{
+	UTypedElementRegistry* Registry = UTypedElementRegistry::GetInstance();
+	if (Registry) // If the registry has already been destroyed there's no point in clearing the reference.
+	{
+		Registry->SetDataStorage(nullptr);
+		Registry->SetDataStorageCompatibility(nullptr);
+	}
+
+	DatabaseCompatibility->Deinitialize();
+	Database->Deinitialize();
+	bInitialized = false;
+}
+
+void FTypedElementsDataStorageModule::AddReferencedObjects(FReferenceCollector& Collector)
+{
+	if (bInitialized)
+	{
+		Collector.AddReferencedObject(Database);
+		Collector.AddReferencedObject(DatabaseCompatibility);
+	}
+}
+
+FString FTypedElementsDataStorageModule::GetReferencerName() const
+{
+	return TEXT("Typed Elements: Data Storage Module");
+}
+
+#undef LOCTEXT_NAMESPACE
+	
+IMPLEMENT_MODULE(FTypedElementsDataStorageModule, TypedElementsDataStorage)
\ No newline at end of file
diff --git a/Engine/Plugins/TypedElementsDataStorage/Source/TypedElementsDataStorage/Public/TypedElementsDataStorage.h b/Engine/Plugins/TypedElementsDataStorage/Source/TypedElementsDataStorage/Public/TypedElementsDataStorage.h
new file mode 100644
--- /dev/null
+++ b/Engine/Plugins/TypedElementsDataStorage/Source/TypedElementsDataStorage/Public/TypedElementsDataStorage.h
@@ -0,0 +1,40 @@
+// Copyright Epic Games, Inc. All Rights Reserved.
+
+#pragma once
+
+#include "CoreMinimal.h"
+#include "Modules/ModuleManager.h"
+#include "UObject/GCObject.h"
+#include "UObject/ObjectPtr.h"
+
+struct FMassEntityManager;
+class FReferenceCollector;
+class ITypedElementDataStorageInterface;
+class UTypedElementDatabase;
+class UTypedElementDatabaseCompatibility;
+
+class FTypedElementsDataStorageModule : public IModuleInterface, public FGCObject
+{
+public:
+	~FTypedElementsDataStorageModule() override = default;
+
+	//
+	// IModuleInterface
+	//
+
+	virtual void StartupModule() override;
+	virtual void ShutdownModule() override;
+
+	
+	//
+	// FGCObject
+	//
+	
+	void AddReferencedObjects(FReferenceCollector& Collector) override;
+	FString GetReferencerName() const override;
+
+private:
+	TObjectPtr<UTypedElementDatabase> Database;
+	TObjectPtr<UTypedElementDatabaseCompatibility> DatabaseCompatibility;
+	bool bInitialized{ false };
+};
diff --git a/Engine/Plugins/TypedElementsDataStorage/Source/TypedElementsDataStorage/TypedElementsDataStorage.Build.cs b/Engine/Plugins/TypedElementsDataStorage/Source/TypedElementsDataStorage/TypedElementsDataStorage.Build.cs
new file mode 100644
--- /dev/null
+++ b/Engine/Plugins/TypedElementsDataStorage/Source/TypedElementsDataStorage/TypedElementsDataStorage.Build.cs
@@ -0,0 +1,32 @@
+// Copyright Epic Games, Inc. All Rights Reserved.
+
+using UnrealBuildTool;
+
+public class TypedElementsDataStorage : ModuleRules
+{
+	public TypedElementsDataStorage(ReadOnlyTargetRules Target) : base(Target)
+	{
+		PCHUsage = ModuleRules.PCHUsageMode.UseExplicitOrSharedPCHs;
+
+		if (Target.bBuildEditor)
+		{
+			PublicIncludePaths.AddRange(new string[] {});
+			PrivateIncludePaths.AddRange(new string[] {});
+
+			PublicDependencyModuleNames.AddRange(
+				new string[]
+				{
+					"Core",
+					"CoreUObject",
+					"Engine",
+					"MassActors",
+					"MassEntity",
+					"TypedElementFramework",
+					"UnrealEd"
+				});
+
+			PrivateDependencyModuleNames.AddRange(new string[] {});
+			DynamicallyLoadedModuleNames.AddRange(new string[] {});
+		}
+	}
+}
diff --git a/Engine/Plugins/TypedElementsDataStorage/TypedElementsDataStorage.uplugin b/Engine/Plugins/TypedElementsDataStorage/TypedElementsDataStorage.uplugin
new file mode 100644
--- /dev/null
+++ b/Engine/Plugins/TypedElementsDataStorage/TypedElementsDataStorage.uplugin
@@ -0,0 +1,33 @@
+{
+	"FileVersion": 3,
+	"Version": 1,
+	"VersionName": "0.1",
+	"FriendlyName": "Typed Elements Data Storage",
+	"Description": "A central extendable data storage for editors and their corresponding data with support for viewing and editing through a collection of widgets.",
+	"Category": "Editor",
+	"CreatedBy": "Epic Games, Inc.",
+	"CreatedByURL": "http://epicgames.com",
+	"DocsURL": "",
+	"MarketplaceURL": "",
+	"SupportURL": "",
+	"CanContainContent": true,
+	"IsBetaVersion": false,
+	"IsExperimentalVersion": true,
+	"Installed": false,
+	"Modules": [
+		{
+			"Name": "TypedElementsDataStorage",
+			"Type": "Editor"
+		}
+	],
+	"Plugins": [
+		{
+			"Name": "MassEntity",
+			"Enabled": true
+		},
+		{
+			"Name": "MassGameplay",
+			"Enabled": true
+		}
+	]
+}
\ No newline at end of file
diff --git a/Engine/Source/Editor/UnrealEd/Private/EditorActor.cpp b/Engine/Source/Editor/UnrealEd/Private/EditorActor.cpp
--- a/Engine/Source/Editor/UnrealEd/Private/EditorActor.cpp
+++ b/Engine/Source/Editor/UnrealEd/Private/EditorActor.cpp
@@ -1187,6 +1187,7 @@ bool UUnrealEdEngine::DeleteActors(const TArray<AActor*>& InActorsToDelete, UWor
 		{
 			InSelectionSet->DeselectElement(ActorHandle, SelectionOptions);
 		}
+		UEngineElementsLibrary::UnregisterActorElement(Actor);
 
 		// Modify the level.  Each level is modified only once.
 		// @todo DB: Shouldn't this be calling UWorld::ModifyLevel?
diff --git a/Engine/Source/Runtime/Engine/Private/Elements/Framework/EngineElementsLibrary.cpp b/Engine/Source/Runtime/Engine/Private/Elements/Framework/EngineElementsLibrary.cpp
--- a/Engine/Source/Runtime/Engine/Private/Elements/Framework/EngineElementsLibrary.cpp
+++ b/Engine/Source/Runtime/Engine/Private/Elements/Framework/EngineElementsLibrary.cpp
@@ -365,11 +365,22 @@ TTypedElementOwner<FActorElementData> UEngineElementsLibrary::CreateActorElement
 	{
 		ActorElement.GetDataChecked().Actor = const_cast<AActor*>(InActor);
 	}
+	
+	UTypedElementRegistry* Registry = UTypedElementRegistry::GetInstance();
+	TObjectPtr<ITypedElementDataStorageCompatibilityInterface> Storage = Registry->GetMutableDataStorageCompatibility();
+	if (Storage)
+	{
+		Storage->AddCompatibleObject(const_cast<AActor*>(InActor));
+	}
+
 	return ActorElement;
 }
 
 void UEngineElementsLibrary::DestroyActorElement(const AActor* InActor, TTypedElementOwner<FActorElementData>& InOutActorElement)
 {
+	UTypedElementRegistry* Registry = UTypedElementRegistry::GetInstance();
+	UnregisterActorElement(InActor);
+
 	if (InOutActorElement)
 	{
 		checkf(InOutActorElement.GetDataChecked().Actor == InActor, TEXT("Actor element was not for this actor instance! %s"), *InActor->GetPathName());
@@ -377,6 +388,16 @@ void UEngineElementsLibrary::DestroyActorElement(const AActor* InActor, TTypedEl
 	}
 }
 
+void UEngineElementsLibrary::UnregisterActorElement(const AActor* InActor)
+{
+	UTypedElementRegistry* Registry = UTypedElementRegistry::GetInstance();
+	TObjectPtr<ITypedElementDataStorageCompatibilityInterface> Storage = Registry->GetMutableDataStorageCompatibility();
+	if (Storage)
+	{
+		Storage->RemoveCompatibleObject(const_cast<AActor*>(InActor));
+	}
+}
+
 #if WITH_EDITOR
 void UEngineElementsLibrary::CreateEditorActorElement(const AActor* Actor)
 {
diff --git a/Engine/Source/Runtime/Engine/Public/Elements/Framework/EngineElementsLibrary.h b/Engine/Source/Runtime/Engine/Public/Elements/Framework/EngineElementsLibrary.h
--- a/Engine/Source/Runtime/Engine/Public/Elements/Framework/EngineElementsLibrary.h
+++ b/Engine/Source/Runtime/Engine/Public/Elements/Framework/EngineElementsLibrary.h
@@ -42,6 +42,7 @@ public:
 
 	static TTypedElementOwner<FActorElementData> CreateActorElement(const AActor* InActor);
 	static void DestroyActorElement(const AActor* InActor, TTypedElementOwner<FActorElementData>& InOutActorElement);
+	static void UnregisterActorElement(const AActor* InActor);
 #if WITH_EDITOR
 	static void CreateEditorActorElement(const AActor* Actor);
 	static void DestroyEditorActorElement(const AActor* Actor);
diff --git a/Engine/Source/Runtime/TypedElementFramework/Private/Elements/Framework/TypedElementRegistry.cpp b/Engine/Source/Runtime/TypedElementFramework/Private/Elements/Framework/TypedElementRegistry.cpp
--- a/Engine/Source/Runtime/TypedElementFramework/Private/Elements/Framework/TypedElementRegistry.cpp
+++ b/Engine/Source/Runtime/TypedElementFramework/Private/Elements/Framework/TypedElementRegistry.cpp
@@ -62,6 +62,36 @@ UTypedElementRegistry* UTypedElementRegistry::GetInstance()
 	return Instance.Get();
 }
 
+ITypedElementDataStorageInterface* UTypedElementRegistry::GetMutableDataStorage()
+{
+	return DataStorage;
+}
+
+const ITypedElementDataStorageInterface* UTypedElementRegistry::GetDataStorage() const
+{
+	return DataStorage;
+}
+
+void UTypedElementRegistry::SetDataStorage(ITypedElementDataStorageInterface* Storage)
+{
+	DataStorage = Storage;
+}
+
+ITypedElementDataStorageCompatibilityInterface* UTypedElementRegistry::GetMutableDataStorageCompatibility()
+{
+	return DataStorageCompatibility;
+}
+
+const ITypedElementDataStorageCompatibilityInterface* UTypedElementRegistry::GetDataStorageCompatibility() const
+{
+	return DataStorageCompatibility;
+}
+
+void UTypedElementRegistry::SetDataStorageCompatibility(ITypedElementDataStorageCompatibilityInterface* Storage)
+{
+	DataStorageCompatibility = Storage;
+}
+
 void UTypedElementRegistry::FinishDestroy()
 {
 	TStrongObjectPtr<UTypedElementRegistry>& Instance = GetTypedElementRegistryInstance();
diff --git a/Engine/Source/Runtime/TypedElementFramework/Public/Elements/Framework/TypedElementRegistry.h b/Engine/Source/Runtime/TypedElementFramework/Public/Elements/Framework/TypedElementRegistry.h
--- a/Engine/Source/Runtime/TypedElementFramework/Public/Elements/Framework/TypedElementRegistry.h
+++ b/Engine/Source/Runtime/TypedElementFramework/Public/Elements/Framework/TypedElementRegistry.h
@@ -17,6 +17,8 @@
 #include "Elements/Framework/TypedElementLimits.h"
 #include "Elements/Framework/TypedElementList.h"
 #include "Elements/Framework/TypedElementListFwd.h"
+#include "Elements/Interfaces/TypedElementDataStorageInterface.h"
+#include "Elements/Interfaces/TypedElementDataStorageCompatibilityInterface.h"
 #include "HAL/CriticalSection.h"
 #include "HAL/PlatformCrt.h"
 #include "Logging/LogVerbosity.h"
@@ -75,6 +77,14 @@ public:
 	UFUNCTION(BlueprintPure, DisplayName="Get Default Typed Element Registry", Category = "TypedElementFramework|Registry", meta=(ScriptName="GetDefaultTypedElementRegistry"))
 	static UTypedElementRegistry* GetInstance();
 
+	ITypedElementDataStorageInterface* GetMutableDataStorage();
+	const ITypedElementDataStorageInterface* GetDataStorage() const;
+	void SetDataStorage(ITypedElementDataStorageInterface* Storage);
+
+	ITypedElementDataStorageCompatibilityInterface* GetMutableDataStorageCompatibility();
+	const ITypedElementDataStorageCompatibilityInterface* GetDataStorageCompatibility() const;
+	void SetDataStorageCompatibility(ITypedElementDataStorageCompatibilityInterface* Storage);
+
 	/**
 	 * Event fired when references to one element should be replaced with a reference to a different element.
 	 */
@@ -671,4 +681,7 @@ private:
 	FOnElementReplaced OnElementReplacedDelegate;
 	FOnElementUpdated OnElementUpdatedDelegate;
 	FSimpleMulticastDelegate OnProcessingDeferredElementsToDestroyDelegate;
+
+	ITypedElementDataStorageInterface* DataStorage = nullptr;
+	ITypedElementDataStorageCompatibilityInterface* DataStorageCompatibility = nullptr;
 };
diff --git a/Engine/Source/Runtime/TypedElementFramework/Public/Elements/Interfaces/TypedElementDataStorageCompatibilityInterface.h b/Engine/Source/Runtime/TypedElementFramework/Public/Elements/Interfaces/TypedElementDataStorageCompatibilityInterface.h
new file mode 100644
--- /dev/null
+++ b/Engine/Source/Runtime/TypedElementFramework/Public/Elements/Interfaces/TypedElementDataStorageCompatibilityInterface.h
@@ -0,0 +1,29 @@
+// Copyright Epic Games, Inc. All Rights Reserved.
+
+#pragma once
+
+#include "UObject/Interface.h"
+#include "UObject/ObjectMacros.h"
+
+#include "TypedElementDataStorageCompatibilityInterface.generated.h"
+
+class AActor;
+
+UINTERFACE(MinimalAPI)
+class UTypedElementDataStorageCompatibilityInterface : public UInterface
+{
+	GENERATED_BODY()
+};
+
+/**
+ * Interface to provide compatibility with existing systems that don't directly
+ * support the data storage.
+ */
+class TYPEDELEMENTFRAMEWORK_API ITypedElementDataStorageCompatibilityInterface
+{
+	GENERATED_BODY()
+
+public:
+	virtual void AddCompatibleObject(AActor* Actor) = 0;
+	virtual void RemoveCompatibleObject(AActor* Actor) = 0;
+};
\ No newline at end of file
diff --git a/Engine/Source/Runtime/TypedElementFramework/Public/Elements/Interfaces/TypedElementDataStorageInterface.h b/Engine/Source/Runtime/TypedElementFramework/Public/Elements/Interfaces/TypedElementDataStorageInterface.h
new file mode 100644
--- /dev/null
+++ b/Engine/Source/Runtime/TypedElementFramework/Public/Elements/Interfaces/TypedElementDataStorageInterface.h
@@ -0,0 +1,134 @@
+// Copyright Epic Games, Inc. All Rights Reserved.
+
+#pragma once
+
+#include "Containers/ContainersFwd.h"
+#include "Delegates/Delegate.h"
+#include "Math/NumericLimits.h"
+#include "UObject/Interface.h"
+#include "UObject/NameTypes.h"
+#include "UObject/ObjectMacros.h"
+
+#include "TypedElementDataStorageInterface.generated.h"
+
+class UClass;
+class UScriptStruct;
+
+using TypedElementTableHandle = uint64;
+static constexpr auto TypedElementInvalidTableHandle = TNumericLimits<TypedElementTableHandle>::Max();
+using TypedElementRowHandle = uint64;
+static constexpr auto TypedElementInvalidRowHandle = TNumericLimits<TypedElementRowHandle>::Max();
+
+using FTypedElementOnDataStorageCreation = FSimpleMulticastDelegate;
+using FTypedElementOnDataStorageDestruction = FSimpleMulticastDelegate;
+using FTypedElementOnDataStorageUpdate = FSimpleMulticastDelegate;
+
+using TypedElementDataStorageCreationCallbackRef = TFunctionRef<void(TypedElementRowHandle Row)>;
+
+UINTERFACE(MinimalAPI)
+class UTypedElementDataStorageInterface : public UInterface
+{
+	GENERATED_BODY()
+};
+
+class TYPEDELEMENTFRAMEWORK_API ITypedElementDataStorageInterface
+{
+	GENERATED_BODY()
+
+public:
+	/**
+	 * Delegate that will be called when the data storage has been created. For various reasons the data storage
+	 * can be destroyed and recreated during the lifetime of the editor. This delegate allows for the opportunity
+	 * to register for this event so information such as table definitions can be registered again.
+	 * Note that by the time this delegate is available there will already be a data storage created. It's therefore
+	 * important that any information is registered with the data storage before registering a callback with the
+	 * delegate.
+	 */
+	virtual FTypedElementOnDataStorageCreation& OnCreation() = 0;
+	/**
+	 * Delegate that will be called when the data storage is about to be destroyed. This can be use as an opportunity
+	 * to store any data in the data storage before it becomes unavailable.
+	 */
+	virtual FTypedElementOnDataStorageDestruction& OnDestruction() = 0;
+	/**
+	 * Called periodically when the storage is available. This provides an opportunity to do any repeated processing
+	 * for the data storage.
+	 */
+	virtual FTypedElementOnDataStorageUpdate& OnUpdate() = 0;
+
+	/**
+	 * Whether or not the data storage is available. The data storage is available most of the time, but can be
+	 * unavailable for a brief time between being destroyed and a new one created.
+	 */
+	virtual bool IsAvailable() const = 0;
+
+	/**
+	 * Tables are automatically created by taking an existing table and adding/removing columns. For
+	 * performance its however better to create a tables before hand and add objects to the table. This
+	 * doesn't prevent those objects from having columns added/removed at a later time.
+	 * To make debugging and profiling easier it's also recommended to add a name.
+	 */
+	virtual TypedElementTableHandle RegisterTable(TConstArrayView<const UScriptStruct*> ColumnList) = 0;
+	virtual TypedElementTableHandle RegisterTable(TConstArrayView<const UScriptStruct*> ColumnList, const FName Name) = 0;
+
+	/**
+	 * Adds a row to a table.
+	 */
+	virtual TypedElementRowHandle AddRow(TypedElementTableHandle Table) = 0;
+	virtual TypedElementRowHandle AddRow(FName TableName) = 0;
+	/**
+	 * Add multiple rows at once. For each new row the OnCreated callback is called. Callers are expected to use the callback to
+	 * initialize the row if needed.
+	 */
+	virtual bool BatchAddRow(TypedElementTableHandle Table, int32 Count, TypedElementDataStorageCreationCallbackRef OnCreated) = 0;
+	virtual bool BatchAddRow(FName TableName, int32 Count, TypedElementDataStorageCreationCallbackRef OnCreated) = 0;
+	
+	/**
+	 * Removes a previously added row.
+	 */
+	virtual void RemoveRow(TypedElementRowHandle Row) = 0;
+
+	/**
+	 * Retrieves a pointer to the column of the given row or a nullptr if not found.
+	 */
+	virtual void* GetColumnData(TypedElementRowHandle Row, const UScriptStruct* ColumnType) = 0;
+
+	/**
+	 * Returns a pointer to the registered external system if found, otherwise null.
+	 */
+	virtual void* GetExternalSystemAddress(UClass* Target) = 0;
+
+	
+	/**
+	 *
+	 * The following are utility functions that are not part of the interface but are provided in order to make using the
+	 * interface easier.
+	 *
+	 */
+	
+	
+	/**
+	 * Returns a pointer to the column of the given row or a nullptr if the type couldn't be found or the row doesn't exist.
+	 */
+	template<typename ColumnType>
+	ColumnType* GetColumn(TypedElementRowHandle Row);
+
+	/**
+	 * Returns a pointer to the registered external system if found, otherwise null.
+	 */
+	template<typename SystemType>
+	SystemType* GetExternalSystem();
+};
+
+// Implementations
+template<typename ColumnType>
+ColumnType* ITypedElementDataStorageInterface::GetColumn(TypedElementRowHandle Row)
+{
+	return reinterpret_cast<ColumnType*>(GetColumnData(Row, ColumnType::StaticStruct()));
+}
+
+template<typename SystemType>
+SystemType* ITypedElementDataStorageInterface::GetExternalSystem()
+{
+	return reinterpret_cast<SystemType*>(GetExternalSystemAddress(SystemType::StaticClass()));
+}
\ No newline at end of file
