Address a warning from the static analyzer.

The static analyzer reported on a potential null pointer dereferencing in the UI framework part of the Typed Elements Data Storage. Checks were added to guard against null pointers.

diff --git a/Engine/Plugins/TypedElementsDataStorage/Source/TypedElementsDataStorage/Private/TypedElementDatabaseUI.cpp b/Engine/Plugins/TypedElementsDataStorage/Source/TypedElementsDataStorage/Private/TypedElementDatabaseUI.cpp
--- a/Engine/Plugins/TypedElementsDataStorage/Source/TypedElementsDataStorage/Private/TypedElementDatabaseUI.cpp
+++ b/Engine/Plugins/TypedElementsDataStorage/Source/TypedElementsDataStorage/Private/TypedElementDatabaseUI.cpp
@@ -61,9 +61,15 @@ void UTypedElementDatabaseUi::CreateWidgetInstanceFromDescription(
 {
 	FTypedElementWidgetConstructor* Constructor = reinterpret_cast<FTypedElementWidgetConstructor*>(
 		FMemory_Alloca_Aligned(Target->GetStructureSize(), Target->GetMinAlignment()));
-	Target->InitializeStruct(Constructor);
-
-	CreateWidgetInstance(*Constructor, Arguments, ConstructionCallback, Target);
+	if (Constructor)
+	{
+		Target->InitializeStruct(Constructor);
+		CreateWidgetInstance(*Constructor, Arguments, ConstructionCallback, Target);
+	}
+	else
+	{
+		checkf(false, TEXT("Remaining stack space is too small to create a Typed Elements widget constructor from a description."));
+	}
 }
 
 void UTypedElementDatabaseUi::CreateWidgetInstanceFromInstance(
@@ -74,10 +80,16 @@ void UTypedElementDatabaseUi::CreateWidgetInstanceFromInstance(
 {
 	FTypedElementWidgetConstructor* Constructor = reinterpret_cast<FTypedElementWidgetConstructor*>(
 		FMemory_Alloca_Aligned(Target->GetStructureSize(), Target->GetMinAlignment()));
-	Target->InitializeStruct(Constructor);
-	Target->CopyScriptStruct(Constructor, SourceConstructor);
-
-	CreateWidgetInstance(*Constructor, Arguments, ConstructionCallback, Target);
+	if (Constructor)
+	{
+		Target->InitializeStruct(Constructor);
+		Target->CopyScriptStruct(Constructor, SourceConstructor);
+		CreateWidgetInstance(*Constructor, Arguments, ConstructionCallback, Target);
+	}
+	else
+	{
+		checkf(false, TEXT("Remaining stack space is too small to create a Typed Elements widget constructor from an instance."));
+	}
 }
 
 void UTypedElementDatabaseUi::CreateWidgetInstance(
