The Typed Elements Data Storage now creates a persistent EditorStorage world.

By introducing a dedicated world the data in the data storage survives level transitions. The new world is constructed to be very minimal as it only starts up the minimally required MASS subsystems and because it only does a "TimeOnly" Tick.

diff --git a/Engine/Plugins/TypedElementsDataStorage/Source/TypedElementsDataStorage/Private/TypedElementDatabase.cpp b/Engine/Plugins/TypedElementsDataStorage/Source/TypedElementsDataStorage/Private/TypedElementDatabase.cpp
--- a/Engine/Plugins/TypedElementsDataStorage/Source/TypedElementsDataStorage/Private/TypedElementDatabase.cpp
+++ b/Engine/Plugins/TypedElementsDataStorage/Source/TypedElementsDataStorage/Private/TypedElementDatabase.cpp
@@ -5,51 +5,78 @@
 #include "Editor.h"
 #include "Engine/World.h"
 #include "MassEntitySubsystem.h"
+#include "Stats/Stats2.h"
+#include "TickTaskManagerInterface.h"
 
 void UTypedElementDatabase::Initialize()
 {
-	ActiveEditorWorld = GEditor->GetEditorWorldContext().World();
-	check(ActiveEditorWorld);
+	constexpr bool bInformEngineOfWorld = false;
+	constexpr bool bAddToRoot = false;
+	// Initialize the world as minimal as possible as it's only used to store data and tick a few subsystems that require updates to 
+	// keep data consistent.
+	UWorld::InitializationValues InitializationValues;
+	InitializationValues
+		.InitializeScenes(false)
+		.AllowAudioPlayback(false)
+		.CreatePhysicsScene(false)
+		.CreateNavigation(false)
+		.CreateAISystem(false)
+		.ShouldSimulatePhysics(false)
+		.CreateFXSystem(false);
+	ActiveEditorWorld = TObjectPtr<UWorld>(UWorld::CreateWorld(EWorldType::EditorStorage, bInformEngineOfWorld,
+		FName(TEXT("Typed Elements Data Storage World")), GetTransientPackage(), bAddToRoot, ERHIFeatureLevel::Num, &InitializationValues));
+	ActiveEditorWorld->SetShouldTick(false); // Ticks will be done explicitly.
+
 	UMassEntitySubsystem* Mass = ActiveEditorWorld->GetSubsystem<UMassEntitySubsystem>();
-	check(Mass);
+	checkf(Mass, TEXT("New editor data world created, but an MASS entity subsystem wasn't created within it."));
 	ActiveEditorEntityManager = Mass->GetMutableEntityManager().AsShared();
-	
-	TickCallHandle = FWorldDelegates::OnWorldTickStart.AddLambda([this](UWorld* World, ELevelTick, float)
-		{
-			if (World == ActiveEditorWorld)
-			{
-				OnUpdateDelegate.Broadcast();
-			}
-		});
-
-	GEditor->OnWorldAdded().AddLambda([this](UWorld* NewWorld)
-		{
-			if (GEditor->GetEditorWorldContext().World() == NewWorld)
-			{
-				ActiveEditorWorld = NewWorld;
-				UMassEntitySubsystem* Mass = NewWorld->GetSubsystem<UMassEntitySubsystem>();
-				check(Mass);
-				ActiveEditorEntityManager = Mass->GetMutableEntityManager().AsShared();
-				
-				OnCreationDelegate.Broadcast();
-			}
-		});
-	GEditor->OnWorldDestroyed().AddLambda([this](UWorld* DestroyedWorld)
-		{
-			if (DestroyedWorld == ActiveEditorWorld)
-			{
-				Reset();
-			}
-		});
 }
 
 void UTypedElementDatabase::Deinitialize()
 {
-	FWorldDelegates::OnWorldTickStart.Remove(TickCallHandle);
-
 	Reset();
 }
 
+ETickableTickType UTypedElementDatabase::GetTickableTickType() const
+{ 
+	return ETickableTickType::Always;
+}
+
+bool UTypedElementDatabase::IsTickableWhenPaused() const
+{
+	return true;
+}
+
+bool UTypedElementDatabase::IsTickableInEditor() const
+{
+	return true;
+}
+
+bool UTypedElementDatabase::IsAllowedToTick() const
+{
+	return !HasAllFlags(RF_ClassDefaultObject) && IsAvailable();
+}
+
+TStatId UTypedElementDatabase::GetStatId() const
+{
+	RETURN_QUICK_DECLARE_CYCLE_STAT(UTypedElementDatabase, STATGROUP_Tickables);
+}
+
+void UTypedElementDatabase::Tick(float DeltaTime)
+{
+	checkf(IsAvailable(), TEXT("Typed Element Database was ticked while it's not ready."));
+	OnUpdateDelegate.Broadcast();
+
+	FTickTaskManagerInterface& TaskManager = FTickTaskManagerInterface::Get();
+	TaskManager.StartFrame(ActiveEditorWorld, DeltaTime, LEVELTICK_TimeOnly, ActiveEditorWorld->GetLevels());
+	for (int Group = 0; Group < TG_MAX; ++Group)
+	{
+		constexpr static bool bBlockTillComplete = true;
+		TaskManager.RunTickGroup(static_cast<ETickingGroup>(Group), bBlockTillComplete);
+	}
+	TaskManager.EndFrame();
+}
+
 TSharedPtr<FMassEntityManager> UTypedElementDatabase::GetActiveMutableEditorEntityManager()
 {
 	return ActiveEditorEntityManager;
@@ -369,6 +396,7 @@ ITypedElementDataStorageInterface::FQueryResult UTypedElementDatabase::RunQuery(
 				Result.Completed = FQueryResult::ECompletion::Fully;
 				break;
 			case FQueryDescription::EActionType::Select:
+				checkf(false, TEXT("Support for this option will be coming in a future update."));
 				Result.Completed = FQueryResult::ECompletion::Unsupported;
 				break;
 			case FQueryDescription::EActionType::Count:
@@ -384,6 +412,7 @@ ITypedElementDataStorageInterface::FQueryResult UTypedElementDatabase::RunQuery(
 		}
 		else
 		{
+			checkf(false, TEXT("Support for this option will be coming in a future update."));
 			Result.Completed = FQueryResult::ECompletion::Unsupported;
 		}
 	}
@@ -395,16 +424,6 @@ ITypedElementDataStorageInterface::FQueryResult UTypedElementDatabase::RunQuery(
 	return Result;
 }
 
-FTypedElementOnDataStorageCreation& UTypedElementDatabase::OnCreation()
-{
-	return OnCreationDelegate;
-}
-
-FTypedElementOnDataStorageDestruction& UTypedElementDatabase::OnDestruction()
-{
-	return OnDestructionDelegate;
-}
-
 FTypedElementOnDataStorageUpdate& UTypedElementDatabase::OnUpdate()
 {
 	return OnUpdateDelegate;
@@ -422,11 +441,10 @@ void* UTypedElementDatabase::GetExternalSystemAddress(UClass* Target)
 
 void UTypedElementDatabase::Reset()
 {
-	OnDestructionDelegate.Broadcast();
-
 	Tables.Reset();
 	TableNameLookup.Reset();
 
+	ActiveEditorWorld->CleanupWorld();
 	ActiveEditorWorld = nullptr;
 	ActiveEditorEntityManager.Reset();
 }
\ No newline at end of file
diff --git a/Engine/Plugins/TypedElementsDataStorage/Source/TypedElementsDataStorage/Private/TypedElementDatabase.h b/Engine/Plugins/TypedElementsDataStorage/Source/TypedElementsDataStorage/Private/TypedElementDatabase.h
--- a/Engine/Plugins/TypedElementsDataStorage/Source/TypedElementsDataStorage/Private/TypedElementDatabase.h
+++ b/Engine/Plugins/TypedElementsDataStorage/Source/TypedElementsDataStorage/Private/TypedElementDatabase.h
@@ -2,12 +2,12 @@
 
 #pragma once
 
-#include "CoreMinimal.h"
 #include "Templates/SharedPointer.h"
 #include "Elements/Interfaces/TypedElementDataStorageInterface.h"
 #include "MassArchetypeTypes.h"
 #include "MassEntityQuery.h"
 #include "TypedElementHandleStore.h"
+#include "UObject/ObjectMacros.h"
 
 #include "TypedElementDatabase.generated.h"
 
@@ -17,22 +17,23 @@ class UWorld;
 UCLASS()
 class TYPEDELEMENTSDATASTORAGE_API UTypedElementDatabase 
 	: public UObject
+	, public FTickableGameObject
 	, public ITypedElementDataStorageInterface
 {
 	GENERATED_BODY()
 public:
-	struct FExtendedQuery
-	{
-		FMassEntityQuery NativeQuery;
-		FQueryDescription::EActionType Action{ FQueryDescription::EActionType::None };
-		bool bSimpleQuery{ false };
-	};
-
 	~UTypedElementDatabase() override = default;
-
+	
 	void Initialize();
 	void Deinitialize();
 
+	ETickableTickType GetTickableTickType() const override;
+	bool IsTickableWhenPaused() const override;
+	bool IsTickableInEditor() const override;
+	bool IsAllowedToTick() const override;
+	TStatId GetStatId() const override;
+	void Tick(float DeltaTime) override;
+
 	TSharedPtr<FMassEntityManager> GetActiveMutableEditorEntityManager();
 	TSharedPtr<const FMassEntityManager> GetActiveEditorEntityManager() const;
 
@@ -58,13 +59,17 @@ public:
 	void UnregisterQuery(TypedElementQueryHandle Query) override;
 	FQueryResult RunQuery(TypedElementQueryHandle Query) override;
 
-	FTypedElementOnDataStorageCreation& OnCreation() override;
-	FTypedElementOnDataStorageDestruction& OnDestruction() override;
 	FTypedElementOnDataStorageUpdate& OnUpdate() override;
 	bool IsAvailable() const override;
 	void* GetExternalSystemAddress(UClass* Target) override;
 
 private:
+	struct FExtendedQuery
+	{
+		FMassEntityQuery NativeQuery;
+		FQueryDescription::EActionType Action{ FQueryDescription::EActionType::None };
+		bool bSimpleQuery{ false };
+	};
 	using QueryStore = TTypedElementHandleStore<FExtendedQuery>;
 
 	void Reset();
@@ -74,12 +79,9 @@ private:
 
 	QueryStore Queries;
 
-	FTypedElementOnDataStorageCreation OnCreationDelegate;
-	FTypedElementOnDataStorageDestruction OnDestructionDelegate;
 	FTypedElementOnDataStorageUpdate OnUpdateDelegate;
 
-	FDelegateHandle TickCallHandle;
-
 	TSharedPtr<FMassEntityManager> ActiveEditorEntityManager;
-	UWorld* ActiveEditorWorld{ nullptr };
+	UPROPERTY(Transient)
+	TObjectPtr<UWorld> ActiveEditorWorld{ nullptr };
 };
\ No newline at end of file
diff --git a/Engine/Plugins/TypedElementsDataStorage/Source/TypedElementsDataStorage/Private/TypedElementDatabaseCompatibility.cpp b/Engine/Plugins/TypedElementsDataStorage/Source/TypedElementsDataStorage/Private/TypedElementDatabaseCompatibility.cpp
--- a/Engine/Plugins/TypedElementsDataStorage/Source/TypedElementsDataStorage/Private/TypedElementDatabaseCompatibility.cpp
+++ b/Engine/Plugins/TypedElementsDataStorage/Source/TypedElementsDataStorage/Private/TypedElementDatabaseCompatibility.cpp
@@ -2,6 +2,7 @@
 
 #include "TypedElementDatabaseCompatibility.h"
 
+#include "Editor.h"
 #include "MassActorSubsystem.h"
 
 void UTypedElementDatabaseCompatibility::Initialize(ITypedElementDataStorageInterface* StorageInterface)
@@ -11,8 +12,6 @@ void UTypedElementDatabaseCompatibility::Initialize(ITypedElementDataStorageInte
 	Storage = StorageInterface;
 	Prepare();
 
-	StorageInterface->OnCreation().AddUObject(this, &UTypedElementDatabaseCompatibility::Prepare);
-	StorageInterface->OnDestruction().AddUObject(this, &UTypedElementDatabaseCompatibility::Reset);
 	StorageInterface->OnUpdate().AddUObject(this, &UTypedElementDatabaseCompatibility::AddPendingActors);
 }
 
@@ -35,10 +34,14 @@ void UTypedElementDatabaseCompatibility::RemoveCompatibleObject(AActor* Actor)
 
 	// If there is no actor subsystem it means that the world has been destroyed, including the MASS instance,
 	// so there's no references to clean up.
-	if (Storage && ActorSubsystem && Storage->IsAvailable() && Actor->GetWorld() == ActorSubsystem->GetWorld())
+	if (Storage && ActorSubsystem && Storage->IsAvailable())
 	{
 		FMassEntityHandle Entity = ActorSubsystem->GetEntityHandleFromActor(Actor);
-		if (Entity.IsValid()) // If there's no entity it may have been deleted earlier, e.g. through an explicit delete.
+		// If there's no entity it may:
+		//	- have been deleted earlier, e.g. through an explicit delete.
+		//	- be an actor that never had a world assigned and was therefore never registered.
+		//	- have registered with a MASS instance in another world, e.g. one created for PIE.
+		if (Entity.IsValid())
 		{
 			auto ActorStore = Storage->GetColumn<FMassActorFragment>(Entity.AsNumber());
 			if (ActorStore && !ActorStore->IsOwnedByMass()) // Only remove actors that were externally created.
@@ -72,11 +75,11 @@ void UTypedElementDatabaseCompatibility::CreateStandardArchetypes()
 
 void UTypedElementDatabaseCompatibility::AddPendingActors()
 {
+	UWorld* EditorWorld = GEditor->GetEditorWorldContext().World();
+
 	// Delay processing until the required systems are available by not clearing the pending actor list.
-	if (ActorsPendingRegistration.Num() > 0 && Storage && Storage->IsAvailable() && ActorSubsystem)
+	if (ActorsPendingRegistration.Num() > 0 && Storage && Storage->IsAvailable() && EditorWorld)
 	{
-		UWorld* TargetWorld = ActorSubsystem->GetWorld();
-
 		// Filter out the actors that are already registered or already destroyed. 
 		// The most common case for this is actors created from within MASS.
 		TWeakObjectPtr<AActor>* ActorBegin = ActorsPendingRegistration.GetData();
@@ -85,9 +88,9 @@ void UTypedElementDatabaseCompatibility::AddPendingActors()
 		while (ActorIt != ActorEnd)
 		{
 			if (
-				(*ActorIt).IsValid() && // If not valid then the actor has already been deleted again.
-				(*ActorIt)->GetWorld() == TargetWorld && // If not equal than this is not an actor for the global data storage.
-				!ActorSubsystem->GetEntityHandleFromActor((*ActorIt).Get()).IsValid() // If true, this actor is likely created within MASS.
+				(*ActorIt).IsValid() && // If not valid or stale then the actor has already been deleted before it could be registered.
+				!ActorSubsystem->GetEntityHandleFromActor((*ActorIt).Get()).IsValid() && // If true, this actor is already registered with the Data Storage.
+				(*ActorIt)->GetWorld() == EditorWorld // Only record actors that are added to the editor to avoid taking PIE into account.
 				) 
 			{
 				++ActorIt;
@@ -110,19 +113,19 @@ void UTypedElementDatabaseCompatibility::AddPendingActors()
 			Storage->BatchAddRow(StandardActorTable, ActorsPendingRegistration.Num(), [this, &ActorIt, ActorEnd](TypedElementRowHandle Row)
 				{
 					auto ActorStore = Storage->GetColumn<FMassActorFragment>(Row);
-					if (ActorStore)
-					{
-						constexpr bool bIsOwnedByMass = false;
-						ActorStore->SetAndUpdateHandleMap(FMassEntityHandle::FromNumber(Row), ActorIt->Get(), bIsOwnedByMass);
-
-						checkf(ActorIt < ActorEnd,
-							TEXT("More (%i) entities were added than were requested (%i)."), ActorEnd - ActorIt, ActorsPendingRegistration.Num());
-						++ActorIt;
-					}
+					checkf(ActorStore, TEXT("Newly created row didn't contain the expected FMassActorFragment."));
+					
+					constexpr bool bIsOwnedByMass = false;
+					ActorStore->SetNoHandleMapUpdate(FMassEntityHandle::FromNumber(Row), ActorIt->Get(), bIsOwnedByMass);
+					ActorSubsystem->SetHandleForActor(ActorIt->Get(), FMassEntityHandle::FromNumber(Row));
+
+					checkf(ActorIt < ActorEnd,
+						TEXT("More (%i) entities were added than were requested (%i)."), ActorEnd - ActorIt, ActorsPendingRegistration.Num());
+					++ActorIt;
 				});
 
 			// Reset the container for next set of actors.
 			ActorsPendingRegistration.Reset();
 		}
 	}
-}
+}
\ No newline at end of file
diff --git a/Engine/Plugins/TypedElementsDataStorage/Source/TypedElementsDataStorage/Private/TypedElementDatabaseCompatibility.h b/Engine/Plugins/TypedElementsDataStorage/Source/TypedElementsDataStorage/Private/TypedElementDatabaseCompatibility.h
--- a/Engine/Plugins/TypedElementsDataStorage/Source/TypedElementsDataStorage/Private/TypedElementDatabaseCompatibility.h
+++ b/Engine/Plugins/TypedElementsDataStorage/Source/TypedElementsDataStorage/Private/TypedElementDatabaseCompatibility.h
@@ -2,10 +2,10 @@
 
 #pragma once
 
-#include "CoreMinimal.h"
 #include "Containers/Array.h"
 #include "Elements/Interfaces/TypedElementDataStorageInterface.h"
 #include "Elements/Interfaces/TypedElementDataStorageCompatibilityInterface.h"
+#include "UObject/ObjectMacros.h"
 #include "UObject/WeakObjectPtr.h"
 
 #include "TypedElementDatabaseCompatibility.generated.h"
diff --git a/Engine/Plugins/TypedElementsDataStorage/Source/TypedElementsDataStorage/Private/TypedElementDatabaseUI.cpp b/Engine/Plugins/TypedElementsDataStorage/Source/TypedElementsDataStorage/Private/TypedElementDatabaseUI.cpp
--- a/Engine/Plugins/TypedElementsDataStorage/Source/TypedElementsDataStorage/Private/TypedElementDatabaseUI.cpp
+++ b/Engine/Plugins/TypedElementsDataStorage/Source/TypedElementsDataStorage/Private/TypedElementDatabaseUI.cpp
@@ -11,8 +11,6 @@ void UTypedElementDatabaseUi::Initialize(ITypedElementDataStorageInterface* Stor
 
 	Storage = StorageInterface;
 	CreateStandardArchetypes();
-
-	StorageInterface->OnCreation().AddUObject(this, &UTypedElementDatabaseUi::CreateStandardArchetypes);
 }
 
 void UTypedElementDatabaseUi::Deinitialize()
diff --git a/Engine/Plugins/TypedElementsDataStorage/Source/TypedElementsDataStorage/Private/TypedElementsDataStorage.cpp b/Engine/Plugins/TypedElementsDataStorage/Source/TypedElementsDataStorage/Private/TypedElementsDataStorage.cpp
--- a/Engine/Plugins/TypedElementsDataStorage/Source/TypedElementsDataStorage/Private/TypedElementsDataStorage.cpp
+++ b/Engine/Plugins/TypedElementsDataStorage/Source/TypedElementsDataStorage/Private/TypedElementsDataStorage.cpp
@@ -73,23 +73,30 @@ void FTypedElementsDataStorageModule::StartupModule()
 				bInitialized = true;
 			}
 		});
+	FCoreDelegates::OnExit.AddRaw(this, &FTypedElementsDataStorageModule::ShutdownModule);
 }
 
 void FTypedElementsDataStorageModule::ShutdownModule()
 {
-	UTypedElementRegistry* Registry = UTypedElementRegistry::GetInstance();
-	if (Registry) // If the registry has already been destroyed there's no point in clearing the reference.
+	if (bInitialized)
 	{
-		Registry->SetDataStorage(nullptr);
-		Registry->SetDataStorageCompatibility(nullptr);
-		Registry->SetDataStorageUi(nullptr);
-	}
+		UTypedElementRegistry* Registry = UTypedElementRegistry::GetInstance();
+		if (Registry) // If the registry has already been destroyed there's no point in clearing the reference.
+		{
+			Registry->SetDataStorage(nullptr);
+			Registry->SetDataStorageCompatibility(nullptr);
+			Registry->SetDataStorageUi(nullptr);
+		}
 
-	DatabaseUi->Deinitialize();
-	DatabaseCompatibility->Deinitialize();
-	Database->Deinitialize();
+		if (UObjectInitialized())
+		{
+			DatabaseUi->Deinitialize();
+			DatabaseCompatibility->Deinitialize();
+			Database->Deinitialize();
+		}
 
-	bInitialized = false;
+		bInitialized = false;
+	}
 }
 
 void FTypedElementsDataStorageModule::AddReferencedObjects(FReferenceCollector& Collector)
diff --git a/Engine/Source/Runtime/TypedElementFramework/Public/Elements/Interfaces/TypedElementDataStorageInterface.h b/Engine/Source/Runtime/TypedElementFramework/Public/Elements/Interfaces/TypedElementDataStorageInterface.h
--- a/Engine/Source/Runtime/TypedElementFramework/Public/Elements/Interfaces/TypedElementDataStorageInterface.h
+++ b/Engine/Source/Runtime/TypedElementFramework/Public/Elements/Interfaces/TypedElementDataStorageInterface.h
@@ -224,22 +224,6 @@ public:
 	 * @section Misc
 	 */
 	
-	/**
-	 * Delegate that will be called when the data storage has been created. For various reasons the data storage
-	 * can be destroyed and recreated during the lifetime of the editor. This delegate allows for the opportunity
-	 * to register for this event so information such as table definitions can be registered again.
-	 * Note that by the time this delegate is available there will already be a data storage created. It's therefore
-	 * important that any information is registered with the data storage before registering a callback with the
-	 * delegate.
-	 */
-	virtual FTypedElementOnDataStorageCreation& OnCreation() = 0;
-	
-	/**
-	 * Delegate that will be called when the data storage is about to be destroyed. This can be use as an opportunity
-	 * to store any data in the data storage before it becomes unavailable.
-	 */
-	virtual FTypedElementOnDataStorageDestruction& OnDestruction() = 0;
-	
 	/**
 	 * Called periodically when the storage is available. This provides an opportunity to do any repeated processing
 	 * for the data storage.
