Created a new module in the Typed Elements Data Storage plugin to house the Typed Elements Widgets.

In this first version a widget has been added that takes a count query as input and list the resulting total in the widget. An actor counter and a Typed Elements Widget counter have been registered to appear in the main frame's status bar.

diff --git a/Engine/Plugins/TypedElementsDataStorage/Source/TypedElementsDataStorage/Private/TypedElementSubsystems.cpp b/Engine/Plugins/TypedElementsDataStorage/Source/TypedElementsDataStorage/Private/TypedElementSubsystems.cpp
new file mode 100644
--- /dev/null
+++ b/Engine/Plugins/TypedElementsDataStorage/Source/TypedElementsDataStorage/Private/TypedElementSubsystems.cpp
@@ -0,0 +1,98 @@
+// Copyright Epic Games, Inc. All Rights Reserved.
+
+#include "TypedElementSubsystems.h"
+
+#include "Elements/Framework/TypedElementRegistry.h"
+
+//
+// UTypedElementDataStorageSubsystem
+//
+
+UTypedElementDataStorageSubsystem::~UTypedElementDataStorageSubsystem()
+{
+	DataStorage = nullptr;
+}
+
+ITypedElementDataStorageInterface* UTypedElementDataStorageSubsystem::Get()
+{
+	if (!DataStorage)
+	{
+		UTypedElementRegistry* Registry = UTypedElementRegistry::GetInstance();
+		checkf(Registry, TEXT("UTypedElementDataStorageSubsystem created before the Typed Elements Registry is available."));
+		DataStorage = Registry->GetMutableDataStorage();
+	}
+	return DataStorage;
+}
+
+const ITypedElementDataStorageInterface* UTypedElementDataStorageSubsystem::Get() const
+{
+	return const_cast<UTypedElementDataStorageSubsystem*>(this)->Get();
+}
+
+bool UTypedElementDataStorageSubsystem::DoesSupportWorldType(const EWorldType::Type WorldType) const
+{
+	return WorldType == EWorldType::EditorStorage || Super::DoesSupportWorldType(WorldType);
+}
+
+
+//
+// UTypedElementDataStorageUiSubsystem
+//
+
+UTypedElementDataStorageUiSubsystem::~UTypedElementDataStorageUiSubsystem()
+{
+	DataStorageUi = nullptr;
+}
+
+ITypedElementDataStorageUiInterface* UTypedElementDataStorageUiSubsystem::Get()
+{
+	if (!DataStorageUi)
+	{
+		UTypedElementRegistry* Registry = UTypedElementRegistry::GetInstance();
+		checkf(Registry, TEXT("UTypedElementDataStorageUiSubsystem created before the Typed Elements Registry is available."));
+		DataStorageUi = Registry->GetMutableDataStorageUi();
+	}
+	return DataStorageUi;
+}
+
+const ITypedElementDataStorageUiInterface* UTypedElementDataStorageUiSubsystem::Get() const
+{
+	return const_cast<UTypedElementDataStorageUiSubsystem*>(this)->Get();
+}
+
+bool UTypedElementDataStorageUiSubsystem::DoesSupportWorldType(const EWorldType::Type WorldType) const
+{
+	return WorldType == EWorldType::EditorStorage || Super::DoesSupportWorldType(WorldType);
+}
+
+
+//
+// UTypedElementDataStorageCompatibilitySubsystem
+//
+
+UTypedElementDataStorageCompatibilitySubsystem::~UTypedElementDataStorageCompatibilitySubsystem()
+{
+	DataStorageCompatibility = nullptr;
+}
+
+ITypedElementDataStorageCompatibilityInterface* UTypedElementDataStorageCompatibilitySubsystem::Get()
+{
+	if (!DataStorageCompatibility)
+	{
+		UTypedElementRegistry* Registry = UTypedElementRegistry::GetInstance();
+		checkf(Registry, TEXT("UTypedElementDataStorageCompatibilitySubsystem created before the Typed Elements Registry is available."));
+
+		DataStorageCompatibility = Registry->GetMutableDataStorageCompatibility();
+	}
+	return DataStorageCompatibility;
+}
+
+const ITypedElementDataStorageCompatibilityInterface* UTypedElementDataStorageCompatibilitySubsystem::Get() const
+{
+	return const_cast<UTypedElementDataStorageCompatibilitySubsystem*>(this)->Get();
+}
+
+bool UTypedElementDataStorageCompatibilitySubsystem::DoesSupportWorldType(const EWorldType::Type WorldType) const
+{
+	return WorldType == EWorldType::EditorStorage || Super::DoesSupportWorldType(WorldType);
+}
diff --git a/Engine/Plugins/TypedElementsDataStorage/Source/TypedElementsDataStorage/Private/UI/TypedElementSlateWidgetReferenceColumnUpdateProcessor.cpp b/Engine/Plugins/TypedElementsDataStorage/Source/TypedElementsDataStorage/Private/UI/TypedElementSlateWidgetReferenceColumnUpdateProcessor.cpp
new file mode 100644
--- /dev/null
+++ b/Engine/Plugins/TypedElementsDataStorage/Source/TypedElementsDataStorage/Private/UI/TypedElementSlateWidgetReferenceColumnUpdateProcessor.cpp
@@ -0,0 +1,61 @@
+// Copyright Epic Games, Inc. All Rights Reserved.
+
+#include "UI/TypedElementSlateWidgetReferenceColumnUpdateProcessor.h"
+
+#include "Elements/Columns/TypedElementSlateWidgetColumns.h"
+#include "MassCommonTypes.h"
+
+UTypedElementSlateWidgetReferenceColumnUpdateProcessor::UTypedElementSlateWidgetReferenceColumnUpdateProcessor()
+	: ColumnRemovalQuery(*this)
+	, RowDeletionQuery(*this)
+{
+	ExecutionFlags = static_cast<int32>(EProcessorExecutionFlags::Editor);
+	ExecutionOrder.ExecuteInGroup = UE::Mass::ProcessorGroupNames::SyncWorldToMass;
+	ProcessingPhase = EMassProcessingPhase::PrePhysics;
+}
+
+void UTypedElementSlateWidgetReferenceColumnUpdateProcessor::ConfigureQueries()
+{
+	ColumnRemovalQuery.AddRequirement(FTypedElementSlateWidgetReferenceColumn::StaticStruct(), EMassFragmentAccess::ReadOnly);
+	ColumnRemovalQuery.AddTagRequirement(*FTypedElementSlateWidgetReferenceDeletesRowTag::StaticStruct(), EMassFragmentPresence::None);
+	
+	RowDeletionQuery.AddRequirement(FTypedElementSlateWidgetReferenceColumn::StaticStruct(), EMassFragmentAccess::ReadOnly);
+	RowDeletionQuery.AddTagRequirement(*FTypedElementSlateWidgetReferenceDeletesRowTag::StaticStruct(), EMassFragmentPresence::All);
+}
+
+void UTypedElementSlateWidgetReferenceColumnUpdateProcessor::Execute(
+	FMassEntityManager& EntityManager, FMassExecutionContext& Context)
+{
+	ColumnRemovalQuery.ForEachEntityChunk(EntityManager, Context, [](FMassExecutionContext& Context)
+		{
+			CheckStatus<false>(Context);
+		});
+
+	RowDeletionQuery.ForEachEntityChunk(EntityManager, Context, [](FMassExecutionContext& Context)
+		{
+			CheckStatus<true>(Context);
+		});
+}
+
+template<bool DeleteRow>
+void UTypedElementSlateWidgetReferenceColumnUpdateProcessor::CheckStatus(FMassExecutionContext& Context)
+{
+	int32 Index = 0;
+	TArrayView<FTypedElementSlateWidgetReferenceColumn> WidgetColumns =
+		Context.GetMutableFragmentView<FTypedElementSlateWidgetReferenceColumn>();
+	for (FTypedElementSlateWidgetReferenceColumn& WidgetColumn : WidgetColumns)
+	{
+		if (!WidgetColumn.Widget.IsValid())
+		{
+			if constexpr (DeleteRow)
+			{
+				Context.Defer().DestroyEntity(Context.GetEntity(Index));
+			}
+			else
+			{
+				Context.Defer().RemoveFragment_RuntimeCheck<FTypedElementSlateWidgetReferenceColumn>(Context.GetEntity(Index));
+			}
+		}
+		++Index;
+	}
+}
diff --git a/Engine/Plugins/TypedElementsDataStorage/Source/TypedElementsDataStorage/Private/UI/TypedElementSlateWidgetReferenceColumnUpdateProcessor.h b/Engine/Plugins/TypedElementsDataStorage/Source/TypedElementsDataStorage/Private/UI/TypedElementSlateWidgetReferenceColumnUpdateProcessor.h
new file mode 100644
--- /dev/null
+++ b/Engine/Plugins/TypedElementsDataStorage/Source/TypedElementsDataStorage/Private/UI/TypedElementSlateWidgetReferenceColumnUpdateProcessor.h
@@ -0,0 +1,33 @@
+// Copyright Epic Games, Inc. All Rights Reserved.
+
+#pragma once
+
+#include "MassObserverProcessor.h"
+#include "MassProcessor.h"
+
+#include "TypedElementSlateWidgetReferenceColumnUpdateProcessor.generated.h"
+
+/**
+ * Processor that checks whether or not a widget still exists. If it has been deleted
+ * then it will remove the column from the Data Storage or deletes the entire row if
+ * the FTypedElementSlateWidgetReferenceDeletesRowTag was found.
+ */
+UCLASS()
+class UTypedElementSlateWidgetReferenceColumnUpdateProcessor : public UMassProcessor
+{
+	GENERATED_BODY()
+
+public:
+	UTypedElementSlateWidgetReferenceColumnUpdateProcessor();
+
+protected:
+	void ConfigureQueries() override;
+	void Execute(FMassEntityManager& EntityManager, FMassExecutionContext& Context) override;
+
+private:
+	template<bool DeleteRow>
+	static void CheckStatus(FMassExecutionContext& Context);
+
+	FMassEntityQuery ColumnRemovalQuery;
+	FMassEntityQuery RowDeletionQuery;
+};
\ No newline at end of file
diff --git a/Engine/Plugins/TypedElementsDataStorage/Source/TypedElementsDataStorage/Public/TypedElementSubsystems.h b/Engine/Plugins/TypedElementsDataStorage/Source/TypedElementsDataStorage/Public/TypedElementSubsystems.h
new file mode 100644
--- /dev/null
+++ b/Engine/Plugins/TypedElementsDataStorage/Source/TypedElementsDataStorage/Public/TypedElementSubsystems.h
@@ -0,0 +1,75 @@
+// Copyright Epic Games, Inc. All Rights Reserved.
+
+#pragma once
+
+#include "Subsystems/WorldSubsystem.h"
+#include "UObject/ObjectMacros.h"
+
+#include "TypedElementSubsystems.generated.h"
+
+class ITypedElementDataStorageInterface;
+class ITypedElementDataStorageUiInterface;
+class ITypedElementDataStorageCompatibilityInterface;
+
+/**
+ * A subsystem to provide alternative access to the Typed Elements Data Storage. This can be used in situations where directly accessing
+ * the Data Storage from the Typed Elements Registry is not recommended, such as for MASS.
+ */
+UCLASS()
+class TYPEDELEMENTSDATASTORAGE_API UTypedElementDataStorageSubsystem : public UWorldSubsystem
+{
+	GENERATED_BODY()
+
+public:
+	~UTypedElementDataStorageSubsystem() override;
+
+	ITypedElementDataStorageInterface* Get();
+	const ITypedElementDataStorageInterface* Get() const;
+
+protected:
+	bool DoesSupportWorldType(const EWorldType::Type WorldType) const override;
+
+	mutable ITypedElementDataStorageInterface* DataStorage{ nullptr };
+};
+
+/**
+ * A subsystem to provide alternative access to the Typed Elements Data Storage UI. This can be used in situations where directly 
+ * accessing the UI from the Typed Elements Registry is not recommended, such as for MASS.
+ */
+UCLASS()
+class TYPEDELEMENTSDATASTORAGE_API UTypedElementDataStorageUiSubsystem : public UWorldSubsystem
+{
+	GENERATED_BODY()
+
+public:
+	~UTypedElementDataStorageUiSubsystem() override;
+
+	ITypedElementDataStorageUiInterface* Get();
+	const ITypedElementDataStorageUiInterface* Get() const;
+
+protected:
+	bool DoesSupportWorldType(const EWorldType::Type WorldType) const override;
+
+	mutable ITypedElementDataStorageUiInterface* DataStorageUi{ nullptr };
+};
+
+/**
+ * A subsystem to provide alternative access to the Typed Elements Data Storage Compatiblity. This can be used in situations where directly 
+ * accessing the Compatiblity extension from the Typed Elements Registry is not recommended, such as for MASS.
+ */
+UCLASS()
+class TYPEDELEMENTSDATASTORAGE_API UTypedElementDataStorageCompatibilitySubsystem : public UWorldSubsystem
+{
+	GENERATED_BODY()
+
+public:
+	~UTypedElementDataStorageCompatibilitySubsystem() override;
+
+	ITypedElementDataStorageCompatibilityInterface* Get();
+	const ITypedElementDataStorageCompatibilityInterface* Get() const;
+
+protected:
+	bool DoesSupportWorldType(const EWorldType::Type WorldType) const override;
+
+	mutable ITypedElementDataStorageCompatibilityInterface* DataStorageCompatibility{ nullptr };
+};
\ No newline at end of file
diff --git a/Engine/Plugins/TypedElementsDataStorage/Source/TypedElementsDataStorageUI/Private/TypedElementCounterWidgetConstructor.cpp b/Engine/Plugins/TypedElementsDataStorage/Source/TypedElementsDataStorageUI/Private/TypedElementCounterWidgetConstructor.cpp
new file mode 100644
--- /dev/null
+++ b/Engine/Plugins/TypedElementsDataStorage/Source/TypedElementsDataStorageUI/Private/TypedElementCounterWidgetConstructor.cpp
@@ -0,0 +1,70 @@
+// Copyright Epic Games, Inc. All Rights Reserved.
+
+#include "TypedElementCounterWidgetConstructor.h"
+
+#include "Elements/Columns/TypedElementSlateWidgetColumns.h"
+#include "Elements/Columns/TypedElementValueCacheColumns.h"
+#include "Elements/Framework/TypedElementQueryBuilder.h"
+#include "Framework/Text/TextLayout.h"
+#include "Layout/Margin.h"
+#include "Widgets/DeclarativeSyntaxSupport.h"
+#include "Widgets/Text/STextBlock.h"
+
+#define LOCTEXT_NAMESPACE "TypedElementUI_CounterWidget"
+
+void FTypedElementCounterWidgetConstructor::Register(
+	ITypedElementDataStorageInterface& DataStorage, ITypedElementDataStorageUiInterface& DataStorageUi)
+{
+	using namespace TypedElementQueryBuilder;
+
+	FName Purpose("LevelEditor.StatusBar.ToolBar");
+
+	TUniquePtr<FTypedElementCounterWidgetConstructor> ActorCounter = MakeUnique<FTypedElementCounterWidgetConstructor>();
+	ActorCounter->LabelText = LOCTEXT("ActorCounterStatusBarLabel", "{0} {0}|plural(one=Actor, other=Actors)");
+	ActorCounter->ToolTipText = LOCTEXT(
+		"ActorCounterStatusBarToolTip",
+		"The total number of actors currently in the editor, excluding PIE/SIE and previews.");
+	ActorCounter->Query = DataStorage.RegisterQuery(
+		Count().
+		Where().
+		All("/Script/MassActors.MassActorFragment"_Type).
+		Compile());
+	DataStorageUi.RegisterWidgetFactory(Purpose, MoveTemp(ActorCounter));
+
+	TUniquePtr<FTypedElementCounterWidgetConstructor> WidgetCounter = MakeUnique<FTypedElementCounterWidgetConstructor>();
+	WidgetCounter->LabelText = LOCTEXT("WidgetCounterStatusBarLabel", "{0} {0}|plural(one=Widget, other=Widgets)");
+	WidgetCounter->ToolTipText = LOCTEXT(
+		"WidgetCounterStatusBarToolTip",
+		"The total number of widgets in the editor hosted through the Typed Element's Data Storage.");
+	WidgetCounter->Query = DataStorage.RegisterQuery(
+		Count().
+		Where().
+		All<FTypedElementSlateWidgetReferenceColumn>().
+		Compile());
+	DataStorageUi.RegisterWidgetFactory(Purpose, MoveTemp(WidgetCounter));
+}
+
+TSharedPtr<SWidget> FTypedElementCounterWidgetConstructor::CreateWidget()
+{
+	return SNew(STextBlock)
+		.Text(FText::Format(LabelText, 0))
+		.Margin(FMargin(4.0f, 0.0f))
+		.Justification(ETextJustify::Center);
+}
+
+void FTypedElementCounterWidgetConstructor::AddColumns(
+	ITypedElementDataStorageInterface* DataStorage, TypedElementRowHandle Row, const TSharedPtr<SWidget>& Widget)
+{
+	Super::AddColumns(DataStorage, Row, Widget);
+
+	FTypedElementCounterWidgetColumn* CounterColumn = DataStorage->AddOrGetColumn<FTypedElementCounterWidgetColumn>(Row);
+	checkf(CounterColumn, TEXT("Added a new FTypedElementCounterWidgetColumn to the Typed Elements Data Storage, but didn't get a valid pointer back."));
+	CounterColumn->LabelTextFormatter = LabelText;
+	CounterColumn->Query = Query;
+
+	FTypedElementU32IntValueCacheColumn* CacheColumn = DataStorage->AddOrGetColumn<FTypedElementU32IntValueCacheColumn>(Row);
+	checkf(CacheColumn, TEXT("Added a new FTypedElementUnsigned32BitIntValueCache to the Typed Elements Data Storage, but didn't get a valid pointer back."));
+	CacheColumn->Value = 0;
+}
+
+#undef LOCTEXT_NAMESPACE
\ No newline at end of file
diff --git a/Engine/Plugins/TypedElementsDataStorage/Source/TypedElementsDataStorageUI/Private/TypedElementCounterWidgetProcessor.cpp b/Engine/Plugins/TypedElementsDataStorage/Source/TypedElementsDataStorageUI/Private/TypedElementCounterWidgetProcessor.cpp
new file mode 100644
--- /dev/null
+++ b/Engine/Plugins/TypedElementsDataStorage/Source/TypedElementsDataStorageUI/Private/TypedElementCounterWidgetProcessor.cpp
@@ -0,0 +1,79 @@
+// Copyright Epic Games, Inc. All Rights Reserved.
+
+#include "TypedElementCounterWidgetProcessor.h"
+
+#include "Elements/Columns/TypedElementSlateWidgetColumns.h"
+#include "Elements/Columns/TypedElementValueCacheColumns.h"
+#include "Elements/Framework/TypedElementQueryBuilder.h"
+#include "Elements/Framework/TypedElementRegistry.h"
+#include "Elements/Interfaces/TypedElementDataStorageUiInterface.h"
+#include "TypedElementCounterWidgetConstructor.h"
+#include "TypedElementSubsystems.h"
+#include "MassCommonTypes.h"
+#include "Widgets/Text/STextBlock.h"
+
+UTypedElementCounterWidgetProcessor::UTypedElementCounterWidgetProcessor()
+	: UiQuery(*this)
+{
+	ExecutionFlags = static_cast<int32>(EProcessorExecutionFlags::Editor);
+	ExecutionOrder.ExecuteInGroup = UE::Mass::ProcessorGroupNames::UpdateWorldFromMass;
+	ProcessingPhase = EMassProcessingPhase::FrameEnd;
+	bRequiresGameThreadExecution = true;
+}
+
+void UTypedElementCounterWidgetProcessor::ConfigureQueries()
+{
+	UiQuery.AddRequirement(FTypedElementSlateWidgetReferenceColumn::StaticStruct(), EMassFragmentAccess::ReadWrite);
+	UiQuery.AddRequirement(FTypedElementU32IntValueCacheColumn::StaticStruct(), EMassFragmentAccess::ReadWrite);
+	UiQuery.AddRequirement(FTypedElementCounterWidgetColumn::StaticStruct(), EMassFragmentAccess::ReadOnly);
+	
+	ProcessorRequirements.AddSubsystemRequirement(UTypedElementDataStorageSubsystem::StaticClass(), EMassFragmentAccess::ReadWrite);
+}
+
+void UTypedElementCounterWidgetProcessor::Execute(FMassEntityManager& EntityManager, FMassExecutionContext& Context)
+{
+	ITypedElementDataStorageInterface* DataInterface = 
+		Context.GetMutableSubsystemChecked<UTypedElementDataStorageSubsystem>(EntityManager.GetWorld()).Get();
+	checkf(DataInterface, TEXT(
+		"FTypedElementsDataStorageUiModule tried to process widgets before the "
+		"Typed Elements Data Storage interface is available."));
+
+	UiQuery.ForEachEntityChunk(EntityManager, Context, [DataInterface](FMassExecutionContext& Context)
+		{
+			FTypedElementSlateWidgetReferenceColumn* WidgetIt = 
+				Context.GetMutableFragmentView<FTypedElementSlateWidgetReferenceColumn>().GetData();
+			FTypedElementU32IntValueCacheColumn* CacheIt = 
+				Context.GetMutableFragmentView<FTypedElementU32IntValueCacheColumn>().GetData();
+			FTypedElementCounterWidgetColumn* CounterIt = 
+				Context.GetMutableFragmentView<FTypedElementCounterWidgetColumn>().GetData();
+
+			const int32 NumEntities = Context.GetNumEntities();
+			for (int32 Index = 0; Index < NumEntities; ++Index)
+			{
+				ITypedElementDataStorageInterface::FQueryResult Result = DataInterface->RunQuery(CounterIt->Query);
+				if (Result.Completed == ITypedElementDataStorageInterface::FQueryResult::ECompletion::Fully &&
+					Result.Count != CacheIt->Value)
+				{
+					checkf(WidgetIt->Widget.IsValid(),
+						TEXT("Referenced widget is not valid. A constructed widget may not have been cleaned up. This can also happen if this "
+						"processor is runningin the same phase as the processors responsible for cleaning up old references."));
+
+					TSharedPtr<SWidget> WidgetPointer = WidgetIt->Widget.Pin();
+					checkf(WidgetPointer, TEXT("Widget pointer was null during processing in Counter Widget Processor. This means "
+						"that the processor that cleans up widgets hasn't been run before this processor."));
+					checkf(WidgetPointer->GetType() == STextBlock::StaticWidgetClass().GetWidgetType(),
+						TEXT("Stored widget with FTypedElementCounterWidgetFragment doesn't match type %s, but was a %s."),
+						*(STextBlock::StaticWidgetClass().GetWidgetType().ToString()),
+						*(WidgetPointer->GetTypeAsString()));
+
+					STextBlock* Widget = static_cast<STextBlock*>(WidgetPointer.Get());
+					Widget->SetText(FText::Format(CounterIt->LabelTextFormatter, Result.Count));
+					CacheIt->Value = Result.Count;
+				}
+
+				++WidgetIt;
+				++CacheIt;
+				++CounterIt;
+			}
+		});
+}
diff --git a/Engine/Plugins/TypedElementsDataStorage/Source/TypedElementsDataStorageUI/Private/TypedElementCounterWidgetProcessor.h b/Engine/Plugins/TypedElementsDataStorage/Source/TypedElementsDataStorageUI/Private/TypedElementCounterWidgetProcessor.h
new file mode 100644
--- /dev/null
+++ b/Engine/Plugins/TypedElementsDataStorage/Source/TypedElementsDataStorageUI/Private/TypedElementCounterWidgetProcessor.h
@@ -0,0 +1,25 @@
+// Copyright Epic Games, Inc. All Rights Reserved.
+
+#pragma once
+
+#include "MassEntityTypes.h"
+#include "MassProcessor.h"
+#include "UObject/ObjectMacros.h"
+
+#include "TypedElementCounterWidgetProcessor.generated.h"
+
+/** Runs the count queries and updates the Slate widgets. */
+UCLASS()
+class TYPEDELEMENTSDATASTORAGEUI_API UTypedElementCounterWidgetProcessor : public UMassProcessor
+{
+	GENERATED_BODY()
+
+public:
+	UTypedElementCounterWidgetProcessor();
+
+protected:
+	void ConfigureQueries() override;
+	void Execute(FMassEntityManager& EntityManager, FMassExecutionContext& Context) override;
+
+	FMassEntityQuery UiQuery;
+};
\ No newline at end of file
diff --git a/Engine/Plugins/TypedElementsDataStorage/Source/TypedElementsDataStorageUI/Private/TypedElementsDataStorageUI.cpp b/Engine/Plugins/TypedElementsDataStorage/Source/TypedElementsDataStorageUI/Private/TypedElementsDataStorageUI.cpp
new file mode 100644
--- /dev/null
+++ b/Engine/Plugins/TypedElementsDataStorage/Source/TypedElementsDataStorageUI/Private/TypedElementsDataStorageUI.cpp
@@ -0,0 +1,111 @@
+// Copyright Epic Games, Inc. All Rights Reserved.
+
+#include "TypedElementsDataStorageUI.h"
+
+#include "Elements/Framework/TypedElementRegistry.h"
+#include "Interfaces/IMainFrameModule.h"
+#include "Misc/CoreDelegates.h"
+#include "ToolMenu.h"
+#include "ToolMenus.h"
+#include "TypedElementCounterWidgetConstructor.h"
+#include "UObject/UObjectGlobals.h"
+#include "Widgets/SWindow.h"
+
+#define LOCTEXT_NAMESPACE "FTypedElementsDataStorageModule"
+
+void FTypedElementsDataStorageUiModule::StartupModule()
+{	
+	FCoreDelegates::OnAllModuleLoadingPhasesComplete.AddLambda(
+		[this]()
+		{
+			UTypedElementRegistry* Registry = UTypedElementRegistry::GetInstance();
+			checkf(Registry, TEXT(
+				"FTypedElementsDataStorageUiModule tried to retrieve the UTypedElementRegistry "
+				"before it was available."));
+			if (Registry->AreDataStorageInterfacesSet())
+			{
+				LateInitialization();
+			}
+			else
+			{
+				Registry->OnDataStorageInterfacesSet().AddStatic(&FTypedElementsDataStorageUiModule::LateInitialization);
+			}
+		});
+}
+
+void FTypedElementsDataStorageUiModule::ShutdownModule()
+{
+}
+
+void FTypedElementsDataStorageUiModule::AddReferencedObjects(FReferenceCollector& Collector)
+{
+}
+
+FString FTypedElementsDataStorageUiModule::GetReferencerName() const
+{
+	return TEXT("Typed Elements: Data Storage UI Module");
+}
+
+void FTypedElementsDataStorageUiModule::LateInitialization()
+{
+	RegisterWidgetFactories();
+	IMainFrameModule::Get().OnMainFrameCreationFinished().AddStatic(&FTypedElementsDataStorageUiModule::SetupMainWindowIntegrations);
+}
+
+void FTypedElementsDataStorageUiModule::RegisterWidgetFactories()
+{
+	UTypedElementRegistry* Registry = UTypedElementRegistry::GetInstance();
+	checkf(Registry, TEXT(
+		"FTypedElementsDataStorageUiModule didn't find the UTypedElementRegistry during widget registration when it should be available."));
+
+	ITypedElementDataStorageInterface* DataInterface = Registry->GetMutableDataStorage();
+	checkf(DataInterface, TEXT(
+		"FTypedElementsDataStorageUiModule tried to register widgets before the "
+		"Typed Elements Data Storage interface is available."));
+	ITypedElementDataStorageUiInterface* UiInterface = Registry->GetMutableDataStorageUi();
+	checkf(UiInterface, TEXT(
+		"FTypedElementsDataStorageUiModule tried to register widgets before the "
+		"Typed Elements Data Storage UI interface is available."));
+
+	FTypedElementCounterWidgetConstructor::Register(*DataInterface, *UiInterface);
+}
+
+void FTypedElementsDataStorageUiModule::SetupMainWindowIntegrations(TSharedPtr<SWindow> ParentWindow, bool bIsRunningStartupDialog)
+{
+	UTypedElementRegistry* Registry = UTypedElementRegistry::GetInstance();
+	checkf(Registry, TEXT(
+		"FTypedElementsDataStorageUiModule didn't find the UTypedElementRegistry during main window integration when it should be available."));
+
+	ITypedElementDataStorageUiInterface* UiInterface = Registry->GetMutableDataStorageUi();
+	checkf(UiInterface, TEXT(
+		"FTypedElementsDataStorageUiModule tried to integrate with the main window before the "
+		"Typed Elements Data Storage UI interface is available."));
+	
+	UToolMenu* Menu = UToolMenus::Get()->ExtendMenu("LevelEditor.StatusBar.ToolBar");
+
+	TArray<TSharedRef<SWidget>> Widgets;
+	UiInterface->ConstructWidgets("LevelEditor.StatusBar.ToolBar", {},
+		[&Widgets](const TSharedRef<SWidget>& NewWidget, TypedElementRowHandle Row)
+		{
+			Widgets.Add(NewWidget);
+		});
+
+	if (!Widgets.IsEmpty())
+	{
+		FToolMenuSection& Section = Menu->AddSection("DataStorageSection");
+		int32 WidgetCount = Widgets.Num();
+
+		Section.AddEntry(FToolMenuEntry::InitWidget("DataStorageStatusBarWidget_0", MoveTemp(Widgets[0]), FText::GetEmpty()));
+		for (int32 I = 1; I < WidgetCount; ++I)
+		{
+			Section.AddSeparator(FName(*FString::Format(TEXT("DataStorageStatusBarWidgetDivider_{0}"), { FString::FromInt(I) })));
+			Section.AddEntry(FToolMenuEntry::InitWidget(
+				FName(*FString::Format(TEXT("DataStorageStatusBarWidget_{0}"), { FString::FromInt(I) })),
+				MoveTemp(Widgets[I]), FText::GetEmpty()));
+		}
+	}
+}
+
+#undef LOCTEXT_NAMESPACE
+	
+IMPLEMENT_MODULE(FTypedElementsDataStorageUiModule, TypedElementsDataStorageUI)
\ No newline at end of file
diff --git a/Engine/Plugins/TypedElementsDataStorage/Source/TypedElementsDataStorageUI/Public/TypedElementCounterWidgetConstructor.h b/Engine/Plugins/TypedElementsDataStorage/Source/TypedElementsDataStorageUI/Public/TypedElementCounterWidgetConstructor.h
new file mode 100644
--- /dev/null
+++ b/Engine/Plugins/TypedElementsDataStorage/Source/TypedElementsDataStorageUI/Public/TypedElementCounterWidgetConstructor.h
@@ -0,0 +1,44 @@
+// Copyright Epic Games, Inc. All Rights Reserved.
+
+#pragma once
+
+#include "Elements/Interfaces/TypedElementDataStorageInterface.h"
+#include "Elements/Interfaces/TypedElementDataStorageUiInterface.h"
+#include "Internationalization/Text.h"
+#include "UObject/ObjectMacros.h"
+
+#include "TypedElementCounterWidgetConstructor.generated.h"
+
+/**
+ * Constructor for the counter widget. The counter widget accepts a "count"-query. The query will be periodically
+ * run and the result is written to a textbox widget after it's been formatted using LabelText. An example for 
+ * LabelText is "{0} {0}|plural(one=MyCounter, other=MyCounters)" which will use "MyCounter" if there's exactly one
+ * entry found and otherwise "MyCounters".
+ */
+USTRUCT()
+struct TYPEDELEMENTSDATASTORAGEUI_API FTypedElementCounterWidgetConstructor : public FTypedElementWidgetConstructor
+{
+	GENERATED_BODY()
+
+public:
+	~FTypedElementCounterWidgetConstructor() override = default;
+
+	FText ToolTipText{ NSLOCTEXT("TypedElementUI_CounterWidget", "Tooltip", "Shows the total number found in the editor.") };
+	FText LabelText{ NSLOCTEXT("TypedElementUI_CounterWidget", "Label", "Counted") };
+	TypedElementQueryHandle Query;
+
+	static void Register(ITypedElementDataStorageInterface& DataStorage, ITypedElementDataStorageUiInterface& DataStorageUi);
+
+protected:
+	TSharedPtr<SWidget> CreateWidget() override;
+	void AddColumns(ITypedElementDataStorageInterface* DataStorage, TypedElementRowHandle Row, const TSharedPtr<SWidget>& Widget) override;
+};
+
+USTRUCT()
+struct TYPEDELEMENTSDATASTORAGEUI_API FTypedElementCounterWidgetColumn : public FTypedElementDataStorageColumn
+{
+	GENERATED_BODY()
+
+	FTextFormat LabelTextFormatter;
+	TypedElementQueryHandle Query;
+};
\ No newline at end of file
diff --git a/Engine/Plugins/TypedElementsDataStorage/Source/TypedElementsDataStorageUI/Public/TypedElementsDataStorageUI.h b/Engine/Plugins/TypedElementsDataStorage/Source/TypedElementsDataStorageUI/Public/TypedElementsDataStorageUI.h
new file mode 100644
--- /dev/null
+++ b/Engine/Plugins/TypedElementsDataStorage/Source/TypedElementsDataStorageUI/Public/TypedElementsDataStorageUI.h
@@ -0,0 +1,37 @@
+// Copyright Epic Games, Inc. All Rights Reserved.
+
+#pragma once
+
+#include "CoreMinimal.h"
+#include "Modules/ModuleManager.h"
+#include "UObject/GCObject.h"
+#include "UObject/ObjectPtr.h"
+
+class FReferenceCollector;
+class SWindow;
+
+class FTypedElementsDataStorageUiModule : public IModuleInterface, public FGCObject
+{
+public:
+	~FTypedElementsDataStorageUiModule() override = default;
+
+	//
+	// IModuleInterface
+	//
+
+	virtual void StartupModule() override;
+	virtual void ShutdownModule() override;
+
+	
+	//
+	// FGCObject
+	//
+	
+	void AddReferencedObjects(FReferenceCollector& Collector) override;
+	FString GetReferencerName() const override;
+
+private:
+	static void LateInitialization();
+	static void RegisterWidgetFactories();
+	static void SetupMainWindowIntegrations(TSharedPtr<SWindow> ParentWindow, bool bIsRunningStartupDialog);
+};
diff --git a/Engine/Plugins/TypedElementsDataStorage/Source/TypedElementsDataStorageUI/TypedElementsDataStorageUI.Build.cs b/Engine/Plugins/TypedElementsDataStorage/Source/TypedElementsDataStorageUI/TypedElementsDataStorageUI.Build.cs
new file mode 100644
--- /dev/null
+++ b/Engine/Plugins/TypedElementsDataStorage/Source/TypedElementsDataStorageUI/TypedElementsDataStorageUI.Build.cs
@@ -0,0 +1,39 @@
+// Copyright Epic Games, Inc. All Rights Reserved.
+
+using UnrealBuildTool;
+
+public class TypedElementsDataStorageUI : ModuleRules
+{
+	public TypedElementsDataStorageUI(ReadOnlyTargetRules Target) : base(Target)
+	{
+		PCHUsage = ModuleRules.PCHUsageMode.UseExplicitOrSharedPCHs;
+
+		if (Target.bBuildEditor)
+		{
+			PublicIncludePaths.AddRange(new string[] {});
+			PrivateIncludePaths.AddRange(new string[] {});
+
+			PublicDependencyModuleNames.AddRange(
+				new string[]
+				{
+					"Core",
+					"CoreUObject",
+					"Engine",
+					"EditorFramework",
+					"Slate",
+					"SlateCore",
+					"ToolMenus",
+					"TypedElementFramework",
+					"TypedElementsDataStorage"
+				});
+			PrivateDependencyModuleNames.AddRange(
+				new string[]
+				{
+					"MassActors",
+					"MassEntity"
+				});
+
+			DynamicallyLoadedModuleNames.AddRange(new string[] {});
+		}
+	}
+}
diff --git a/Engine/Plugins/TypedElementsDataStorage/TypedElementsDataStorage.uplugin b/Engine/Plugins/TypedElementsDataStorage/TypedElementsDataStorage.uplugin
--- a/Engine/Plugins/TypedElementsDataStorage/TypedElementsDataStorage.uplugin
+++ b/Engine/Plugins/TypedElementsDataStorage/TypedElementsDataStorage.uplugin
@@ -18,6 +18,10 @@
 		{
 			"Name": "TypedElementsDataStorage",
 			"Type": "Editor"
+		},
+		{
+			"Name": "TypedElementsDataStorageUI",
+			"Type": "Editor"
 		}
 	],
 	"Plugins": [
diff --git a/Engine/Source/Runtime/TypedElementFramework/Public/Elements/Columns/TypedElementValueCacheColumns.h b/Engine/Source/Runtime/TypedElementFramework/Public/Elements/Columns/TypedElementValueCacheColumns.h
new file mode 100644
--- /dev/null
+++ b/Engine/Source/Runtime/TypedElementFramework/Public/Elements/Columns/TypedElementValueCacheColumns.h
@@ -0,0 +1,42 @@
+// Copyright Epic Games, Inc. All Rights Reserved.
+
+#pragma once
+
+#include "Elements/Interfaces/TypedElementDataStorageInterface.h"
+#include "UObject/ObjectMacros.h"
+
+#include "TypedElementValueCacheColumns.generated.h"
+
+
+/**
+ * Column that can be used to cache an unsigned 32-bit value in.
+ */
+USTRUCT()
+struct FTypedElementU32IntValueCacheColumn : public FTypedElementDataStorageColumn
+{
+	GENERATED_BODY()
+
+	uint32 Value;
+};
+
+/**
+ * Column that can be used to cache a signed 32-bit value in.
+ */
+USTRUCT()
+struct FTypedElementI32IntValueCacheColumn : public FTypedElementDataStorageColumn
+{
+	GENERATED_BODY()
+
+	int32 Value;
+};
+
+/**
+ * Column that can be used to cache a 32-bit floating point value in.
+ */
+USTRUCT()
+struct FTypedElementFloatValueCacheColumn : public FTypedElementDataStorageColumn
+{
+	GENERATED_BODY()
+
+	float Value;
+};
\ No newline at end of file
