Added ability to duplicate tables with new columns added plus the ability to find a table by name.

diff --git a/Engine/Plugins/TypedElementsDataStorage/Source/TypedElementsDataStorage/Private/TypedElementDatabase.cpp b/Engine/Plugins/TypedElementsDataStorage/Source/TypedElementsDataStorage/Private/TypedElementDatabase.cpp
--- a/Engine/Plugins/TypedElementsDataStorage/Source/TypedElementsDataStorage/Private/TypedElementDatabase.cpp
+++ b/Engine/Plugins/TypedElementsDataStorage/Source/TypedElementsDataStorage/Private/TypedElementDatabase.cpp
@@ -94,28 +94,47 @@ TypedElementTableHandle UTypedElementDatabase::RegisterTable(TConstArrayView<con
 
 TypedElementTableHandle UTypedElementDatabase::RegisterTable(TConstArrayView<const UScriptStruct*> ColumnList, const FName Name)
 {
-	if (ActiveEditorEntityManager)
+	if (ActiveEditorEntityManager && (!Name.IsValid() || !TableNameLookup.Contains(Name)))
 	{
+		TypedElementTableHandle Result = Tables.Num();
+		Tables.Add(ActiveEditorEntityManager->CreateArchetype(ColumnList, Name));
 		if (Name.IsValid())
 		{
-			if (!TableNameLookup.Contains(Name))
-			{
-				TypedElementTableHandle Result = Tables.Num();
-				Tables.Add(ActiveEditorEntityManager->CreateArchetype(ColumnList, Name));
-				TableNameLookup.Add(Name, Result);
-				return Result;
-			}
+			TableNameLookup.Add(Name, Result);
 		}
-		else
+		return Result;
+	}
+	return TypedElementInvalidTableHandle;
+}
+
+TypedElementTableHandle UTypedElementDatabase::RegisterTable(TypedElementTableHandle SourceTable,
+	TConstArrayView<const UScriptStruct*> ColumnList)
+{
+	return RegisterTable(SourceTable, ColumnList, {});
+}
+
+TypedElementTableHandle UTypedElementDatabase::RegisterTable(TypedElementTableHandle SourceTable, 
+	TConstArrayView<const UScriptStruct*> ColumnList, const FName Name)
+{
+	if (ActiveEditorEntityManager && (!Name.IsValid() || !TableNameLookup.Contains(Name)) && SourceTable < Tables.Num())
+	{
+		TypedElementTableHandle Result = Tables.Num();
+		Tables.Add(ActiveEditorEntityManager->CreateArchetype(Tables[SourceTable], ColumnList, Name));
+		if (Name.IsValid())
 		{
-			TypedElementTableHandle Result = Tables.Num();
-			Tables.Add(ActiveEditorEntityManager->CreateArchetype(ColumnList, Name));
-			return Result;
+			TableNameLookup.Add(Name, Result);
 		}
+		return Result;
 	}
 	return TypedElementInvalidTableHandle;
 }
 
+TypedElementTableHandle UTypedElementDatabase::FindTable(const FName Name)
+{
+	TypedElementTableHandle* TableHandle = TableNameLookup.Find(Name);
+	return TableHandle ? *TableHandle : TypedElementInvalidTableHandle;
+}
+
 TypedElementRowHandle UTypedElementDatabase::AddRow(TypedElementTableHandle Table)
 {
 	checkf(Table < Tables.Num(), TEXT("Attempting to add a row to a non-existing table."));
diff --git a/Engine/Plugins/TypedElementsDataStorage/Source/TypedElementsDataStorage/Private/TypedElementDatabase.h b/Engine/Plugins/TypedElementsDataStorage/Source/TypedElementsDataStorage/Private/TypedElementDatabase.h
--- a/Engine/Plugins/TypedElementsDataStorage/Source/TypedElementsDataStorage/Private/TypedElementDatabase.h
+++ b/Engine/Plugins/TypedElementsDataStorage/Source/TypedElementsDataStorage/Private/TypedElementDatabase.h
@@ -39,6 +39,10 @@ public:
 
 	TypedElementTableHandle RegisterTable(TConstArrayView<const UScriptStruct*> ColumnList) override;
 	TypedElementTableHandle RegisterTable(TConstArrayView<const UScriptStruct*> ColumnList, const FName Name) override;
+	TypedElementTableHandle RegisterTable(TypedElementTableHandle SourceTable, TConstArrayView<const UScriptStruct*> ColumnList) override;
+	TypedElementTableHandle RegisterTable(TypedElementTableHandle SourceTable, TConstArrayView<const UScriptStruct*> ColumnList, 
+		const FName Name) override;
+	TypedElementTableHandle FindTable(const FName Name) override;
 
 	TypedElementRowHandle AddRow(TypedElementTableHandle Table) override;
 	TypedElementRowHandle AddRow(FName TableName) override;
diff --git a/Engine/Source/Runtime/TypedElementFramework/Public/Elements/Interfaces/TypedElementDataStorageInterface.h b/Engine/Source/Runtime/TypedElementFramework/Public/Elements/Interfaces/TypedElementDataStorageInterface.h
--- a/Engine/Source/Runtime/TypedElementFramework/Public/Elements/Interfaces/TypedElementDataStorageInterface.h
+++ b/Engine/Source/Runtime/TypedElementFramework/Public/Elements/Interfaces/TypedElementDataStorageInterface.h
@@ -79,7 +79,17 @@ public:
 	/** Creates a new table for with the provided columns. Optionally a name can be given which is useful for retrieval later. */
 	virtual TypedElementTableHandle RegisterTable(TConstArrayView<const UScriptStruct*> ColumnList) = 0;
 	virtual TypedElementTableHandle RegisterTable(TConstArrayView<const UScriptStruct*> ColumnList, const FName Name) = 0;
+	/** 
+	 * Copies the column information from the provided table and creates a new table for with the provided columns. Optionally a 
+	 * name can be given which is useful for retrieval later.
+	 */
+	virtual TypedElementTableHandle RegisterTable(TypedElementTableHandle SourceTable, 
+		TConstArrayView<const UScriptStruct*> ColumnList) = 0;
+	virtual TypedElementTableHandle RegisterTable(TypedElementTableHandle SourceTable, 
+		TConstArrayView<const UScriptStruct*> ColumnList, const FName Name) = 0;
 
+	/** Returns a previously created table with the provided name or TypedElementInvalidTableHandle if not found. */
+	virtual TypedElementTableHandle FindTable(const FName Name) = 0;
 	
 	/**
 	 * @section Row management
