Changed to TypedElementsDataStorage taking advantage of new World-less Mass approach (no longer needing a dedicated EWorldType::EditorStorage world type).
To maintain current TypedElementsDataStorage behavior also added dedicated delegate to MassEntityEditorSubsystem that gets broadcast just before issuing the Mass phases' execution tasks. TypedElementsDataStorage's logic seems to rely on that.

diff --git a/Engine/Plugins/Runtime/MassEntity/Source/MassEntityEditor/Private/MassEntityEditorSubsystem.cpp b/Engine/Plugins/Runtime/MassEntity/Source/MassEntityEditor/Private/MassEntityEditorSubsystem.cpp
--- a/Engine/Plugins/Runtime/MassEntity/Source/MassEntityEditor/Private/MassEntityEditorSubsystem.cpp
+++ b/Engine/Plugins/Runtime/MassEntity/Source/MassEntityEditor/Private/MassEntityEditorSubsystem.cpp
@@ -128,6 +128,8 @@ void UMassEntityEditorSubsystem::Tick(float DeltaTime)
 		CompletionEvent->Wait();
 	}
 
+	OnPreTickDelegate.Broadcast(DeltaTime);
+
 	for (int PhaseIndex = 0; PhaseIndex < (int)EMassProcessingPhase::MAX; ++PhaseIndex)
 	{
 		const FGraphEventArray Prerequisites = { CompletionEvent };
diff --git a/Engine/Plugins/Runtime/MassEntity/Source/MassEntityEditor/Public/MassEntityEditorSubsystem.h b/Engine/Plugins/Runtime/MassEntity/Source/MassEntityEditor/Public/MassEntityEditorSubsystem.h
--- a/Engine/Plugins/Runtime/MassEntity/Source/MassEntityEditor/Public/MassEntityEditorSubsystem.h
+++ b/Engine/Plugins/Runtime/MassEntity/Source/MassEntityEditor/Public/MassEntityEditorSubsystem.h
@@ -19,11 +19,15 @@ class MASSENTITYEDITOR_API UMassEntityEditorSubsystem : public UEditorSubsystem
 	GENERATED_BODY()
 
 public:
+	DECLARE_MULTICAST_DELEGATE_OneParam(FOnPreTick, float /*DeltaTime*/);
+
 	UMassEntityEditorSubsystem();
 	~UMassEntityEditorSubsystem();
 
 	FMassEntityManager& GetMutableEntityManager() { return EntityManager.Get(); }
 
+	FOnPreTick& GetOnPreTickDelegate() { return OnPreTickDelegate; }
+
 protected:
 	// USubsystem interface begin
 	virtual void Initialize(FSubsystemCollectionBase& Collection) override;
@@ -42,4 +46,7 @@ protected:
 	FGraphEventRef CompletionEvent;
 
 	FMassProcessingPhaseConfig ProcessingPhasesConfig[(uint8)EMassProcessingPhase::MAX];
+
+	/** Gets broadcast right before issuing FMassProcessingPhaseManager's phases execution */
+	FOnPreTick OnPreTickDelegate;
 };
diff --git a/Engine/Plugins/TypedElementsDataStorage/Source/TypedElementsDataStorage/Private/TypedElementDatabase.cpp b/Engine/Plugins/TypedElementsDataStorage/Source/TypedElementsDataStorage/Private/TypedElementDatabase.cpp
--- a/Engine/Plugins/TypedElementsDataStorage/Source/TypedElementsDataStorage/Private/TypedElementDatabase.cpp
+++ b/Engine/Plugins/TypedElementsDataStorage/Source/TypedElementsDataStorage/Private/TypedElementDatabase.cpp
@@ -4,31 +4,18 @@
 
 #include "Editor.h"
 #include "Engine/World.h"
-#include "MassEntitySubsystem.h"
+#include "MassEntityEditorSubsystem.h"
 #include "Stats/Stats2.h"
 #include "TickTaskManagerInterface.h"
+#include "MassSubsystemAccess.h"
 
 void UTypedElementDatabase::Initialize()
 {
-	constexpr bool bInformEngineOfWorld = false;
-	constexpr bool bAddToRoot = false;
-	// Initialize the world as minimal as possible as it's only used to store data and tick a few subsystems that require updates to 
-	// keep data consistent.
-	UWorld::InitializationValues InitializationValues;
-	InitializationValues
-		.InitializeScenes(false)
-		.AllowAudioPlayback(false)
-		.CreatePhysicsScene(false)
-		.CreateNavigation(false)
-		.CreateAISystem(false)
-		.ShouldSimulatePhysics(false)
-		.CreateFXSystem(false);
-	ActiveEditorWorld = TObjectPtr<UWorld>(UWorld::CreateWorld(EWorldType::EditorStorage, bInformEngineOfWorld,
-		FName(TEXT("Typed Elements Data Storage World")), GetTransientPackage(), bAddToRoot, ERHIFeatureLevel::Num, &InitializationValues));
-	ActiveEditorWorld->SetShouldTick(false); // Ticks will be done explicitly.
-
-	UMassEntitySubsystem* Mass = ActiveEditorWorld->GetSubsystem<UMassEntitySubsystem>();
-	checkf(Mass, TEXT("New editor data world created, but an MASS entity subsystem wasn't created within it."));
+	check(GEditor);
+	UMassEntityEditorSubsystem* Mass = GEditor->GetEditorSubsystem<UMassEntityEditorSubsystem>();
+	check(Mass);
+	Mass->GetOnPreTickDelegate().AddUObject(this, &UTypedElementDatabase::OnPreMassTick);
+
 	ActiveEditorEntityManager = Mass->GetMutableEntityManager().AsShared();
 }
 
@@ -37,44 +24,10 @@ void UTypedElementDatabase::Deinitialize()
 	Reset();
 }
 
-ETickableTickType UTypedElementDatabase::GetTickableTickType() const
-{ 
-	return ETickableTickType::Always;
-}
-
-bool UTypedElementDatabase::IsTickableWhenPaused() const
-{
-	return true;
-}
-
-bool UTypedElementDatabase::IsTickableInEditor() const
-{
-	return true;
-}
-
-bool UTypedElementDatabase::IsAllowedToTick() const
-{
-	return !HasAllFlags(RF_ClassDefaultObject) && IsAvailable();
-}
-
-TStatId UTypedElementDatabase::GetStatId() const
-{
-	RETURN_QUICK_DECLARE_CYCLE_STAT(UTypedElementDatabase, STATGROUP_Tickables);
-}
-
-void UTypedElementDatabase::Tick(float DeltaTime)
+void UTypedElementDatabase::OnPreMassTick(float DeltaTime)
 {
 	checkf(IsAvailable(), TEXT("Typed Element Database was ticked while it's not ready."));
 	OnUpdateDelegate.Broadcast();
-
-	FTickTaskManagerInterface& TaskManager = FTickTaskManagerInterface::Get();
-	TaskManager.StartFrame(ActiveEditorWorld, DeltaTime, LEVELTICK_TimeOnly, ActiveEditorWorld->GetLevels());
-	for (int Group = 0; Group < TG_MAX; ++Group)
-	{
-		constexpr static bool bBlockTillComplete = true;
-		TaskManager.RunTickGroup(static_cast<ETickingGroup>(Group), bBlockTillComplete);
-	}
-	TaskManager.EndFrame();
 }
 
 TSharedPtr<FMassEntityManager> UTypedElementDatabase::GetActiveMutableEditorEntityManager()
@@ -450,20 +403,21 @@ FTypedElementOnDataStorageUpdate& UTypedElementDatabase::OnUpdate()
 
 bool UTypedElementDatabase::IsAvailable() const
 {
-	return ActiveEditorWorld && ActiveEditorEntityManager;
+	return bool(ActiveEditorEntityManager);
 }
 
 void* UTypedElementDatabase::GetExternalSystemAddress(UClass* Target)
 {
-	return ActiveEditorWorld ? ActiveEditorWorld->GetSubsystemBase(Target) : nullptr;
+	if (Target && Target->IsChildOf<USubsystem>())
+	{
+		return FMassSubsystemAccess::FetchSubsystemInstance(/*World=*/nullptr, Target);
+	}
+	return nullptr;
 }
 
 void UTypedElementDatabase::Reset()
 {
 	Tables.Reset();
 	TableNameLookup.Reset();
-
-	ActiveEditorWorld->CleanupWorld();
-	ActiveEditorWorld = nullptr;
 	ActiveEditorEntityManager.Reset();
 }
\ No newline at end of file
diff --git a/Engine/Plugins/TypedElementsDataStorage/Source/TypedElementsDataStorage/Private/TypedElementDatabase.h b/Engine/Plugins/TypedElementsDataStorage/Source/TypedElementsDataStorage/Private/TypedElementDatabase.h
--- a/Engine/Plugins/TypedElementsDataStorage/Source/TypedElementsDataStorage/Private/TypedElementDatabase.h
+++ b/Engine/Plugins/TypedElementsDataStorage/Source/TypedElementsDataStorage/Private/TypedElementDatabase.h
@@ -17,7 +17,6 @@ class UWorld;
 UCLASS()
 class TYPEDELEMENTSDATASTORAGE_API UTypedElementDatabase 
 	: public UObject
-	, public FTickableGameObject
 	, public ITypedElementDataStorageInterface
 {
 	GENERATED_BODY()
@@ -27,12 +26,8 @@ public:
 	void Initialize();
 	void Deinitialize();
 
-	ETickableTickType GetTickableTickType() const override;
-	bool IsTickableWhenPaused() const override;
-	bool IsTickableInEditor() const override;
-	bool IsAllowedToTick() const override;
-	TStatId GetStatId() const override;
-	void Tick(float DeltaTime) override;
+	/** Triggered just before underlying Mass processing gets ticked */
+	void OnPreMassTick(float DeltaTime);
 
 	TSharedPtr<FMassEntityManager> GetActiveMutableEditorEntityManager();
 	TSharedPtr<const FMassEntityManager> GetActiveEditorEntityManager() const;
@@ -86,6 +81,4 @@ private:
 	FTypedElementOnDataStorageUpdate OnUpdateDelegate;
 
 	TSharedPtr<FMassEntityManager> ActiveEditorEntityManager;
-	UPROPERTY(Transient)
-	TObjectPtr<UWorld> ActiveEditorWorld{ nullptr };
 };
\ No newline at end of file
diff --git a/Engine/Plugins/TypedElementsDataStorage/Source/TypedElementsDataStorage/Private/TypedElementDatabaseCompatibility.cpp b/Engine/Plugins/TypedElementsDataStorage/Source/TypedElementsDataStorage/Private/TypedElementDatabaseCompatibility.cpp
--- a/Engine/Plugins/TypedElementsDataStorage/Source/TypedElementsDataStorage/Private/TypedElementDatabaseCompatibility.cpp
+++ b/Engine/Plugins/TypedElementsDataStorage/Source/TypedElementsDataStorage/Private/TypedElementDatabaseCompatibility.cpp
@@ -3,6 +3,7 @@
 #include "TypedElementDatabaseCompatibility.h"
 
 #include "Editor.h"
+#include "MassActorEditorSubsystem.h"
 #include "MassActorSubsystem.h"
 
 void UTypedElementDatabaseCompatibility::Initialize(ITypedElementDataStorageInterface* StorageInterface)
@@ -55,7 +56,9 @@ void UTypedElementDatabaseCompatibility::RemoveCompatibleObject(AActor* Actor)
 
 void UTypedElementDatabaseCompatibility::Prepare()
 {
-	ActorSubsystem = Storage->GetExternalSystem<UMassActorSubsystem>();
+	UMassActorEditorSubsystem* MassActorEditorSubsystem = Storage->GetExternalSystem<UMassActorEditorSubsystem>();
+	check(MassActorEditorSubsystem);
+	ActorSubsystem = MassActorEditorSubsystem->GetMutableActorManager().AsShared();
 
 	CreateStandardArchetypes();
 }
diff --git a/Engine/Plugins/TypedElementsDataStorage/Source/TypedElementsDataStorage/Private/TypedElementDatabaseCompatibility.h b/Engine/Plugins/TypedElementsDataStorage/Source/TypedElementsDataStorage/Private/TypedElementDatabaseCompatibility.h
--- a/Engine/Plugins/TypedElementsDataStorage/Source/TypedElementsDataStorage/Private/TypedElementDatabaseCompatibility.h
+++ b/Engine/Plugins/TypedElementsDataStorage/Source/TypedElementsDataStorage/Private/TypedElementDatabaseCompatibility.h
@@ -12,7 +12,7 @@
 
 class AActor;
 class ITypedElementDataStorageInterface;
-class UMassActorSubsystem;
+struct FMassActorManager;
 
 UCLASS()
 class TYPEDELEMENTSDATASTORAGE_API UTypedElementDatabaseCompatibility
@@ -40,5 +40,5 @@ private:
 	
 	TypedElementTableHandle StandardActorTable{ TypedElementInvalidTableHandle };
 	ITypedElementDataStorageInterface* Storage{ nullptr };
-	UMassActorSubsystem* ActorSubsystem{ nullptr };
+	TSharedPtr<FMassActorManager> ActorSubsystem;
 };
\ No newline at end of file
diff --git a/Engine/Plugins/TypedElementsDataStorage/Source/TypedElementsDataStorage/Private/TypedElementSubsystems.cpp b/Engine/Plugins/TypedElementsDataStorage/Source/TypedElementsDataStorage/Private/TypedElementSubsystems.cpp
--- a/Engine/Plugins/TypedElementsDataStorage/Source/TypedElementsDataStorage/Private/TypedElementSubsystems.cpp
+++ b/Engine/Plugins/TypedElementsDataStorage/Source/TypedElementsDataStorage/Private/TypedElementSubsystems.cpp
@@ -29,11 +29,6 @@ const ITypedElementDataStorageInterface* UTypedElementDataStorageSubsystem::Get(
 	return const_cast<UTypedElementDataStorageSubsystem*>(this)->Get();
 }
 
-bool UTypedElementDataStorageSubsystem::DoesSupportWorldType(const EWorldType::Type WorldType) const
-{
-	return WorldType == EWorldType::EditorStorage || Super::DoesSupportWorldType(WorldType);
-}
-
 
 //
 // UTypedElementDataStorageUiSubsystem
@@ -60,11 +55,6 @@ const ITypedElementDataStorageUiInterface* UTypedElementDataStorageUiSubsystem::
 	return const_cast<UTypedElementDataStorageUiSubsystem*>(this)->Get();
 }
 
-bool UTypedElementDataStorageUiSubsystem::DoesSupportWorldType(const EWorldType::Type WorldType) const
-{
-	return WorldType == EWorldType::EditorStorage || Super::DoesSupportWorldType(WorldType);
-}
-
 
 //
 // UTypedElementDataStorageCompatibilitySubsystem
@@ -91,8 +81,3 @@ const ITypedElementDataStorageCompatibilityInterface* UTypedElementDataStorageCo
 {
 	return const_cast<UTypedElementDataStorageCompatibilitySubsystem*>(this)->Get();
 }
-
-bool UTypedElementDataStorageCompatibilitySubsystem::DoesSupportWorldType(const EWorldType::Type WorldType) const
-{
-	return WorldType == EWorldType::EditorStorage || Super::DoesSupportWorldType(WorldType);
-}
diff --git a/Engine/Plugins/TypedElementsDataStorage/Source/TypedElementsDataStorage/Public/TypedElementSubsystems.h b/Engine/Plugins/TypedElementsDataStorage/Source/TypedElementsDataStorage/Public/TypedElementSubsystems.h
--- a/Engine/Plugins/TypedElementsDataStorage/Source/TypedElementsDataStorage/Public/TypedElementSubsystems.h
+++ b/Engine/Plugins/TypedElementsDataStorage/Source/TypedElementsDataStorage/Public/TypedElementSubsystems.h
@@ -2,7 +2,7 @@
 
 #pragma once
 
-#include "Subsystems/WorldSubsystem.h"
+#include "EditorSubsystem.h"
 #include "UObject/ObjectMacros.h"
 
 #include "TypedElementSubsystems.generated.h"
@@ -16,7 +16,7 @@ class ITypedElementDataStorageCompatibilityInterface;
  * the Data Storage from the Typed Elements Registry is not recommended, such as for MASS.
  */
 UCLASS()
-class TYPEDELEMENTSDATASTORAGE_API UTypedElementDataStorageSubsystem : public UWorldSubsystem
+class TYPEDELEMENTSDATASTORAGE_API UTypedElementDataStorageSubsystem : public UEditorSubsystem
 {
 	GENERATED_BODY()
 
@@ -27,8 +27,6 @@ public:
 	const ITypedElementDataStorageInterface* Get() const;
 
 protected:
-	bool DoesSupportWorldType(const EWorldType::Type WorldType) const override;
-
 	mutable ITypedElementDataStorageInterface* DataStorage{ nullptr };
 };
 
@@ -37,7 +35,7 @@ protected:
  * accessing the UI from the Typed Elements Registry is not recommended, such as for MASS.
  */
 UCLASS()
-class TYPEDELEMENTSDATASTORAGE_API UTypedElementDataStorageUiSubsystem : public UWorldSubsystem
+class TYPEDELEMENTSDATASTORAGE_API UTypedElementDataStorageUiSubsystem : public UEditorSubsystem
 {
 	GENERATED_BODY()
 
@@ -48,8 +46,6 @@ public:
 	const ITypedElementDataStorageUiInterface* Get() const;
 
 protected:
-	bool DoesSupportWorldType(const EWorldType::Type WorldType) const override;
-
 	mutable ITypedElementDataStorageUiInterface* DataStorageUi{ nullptr };
 };
 
@@ -58,7 +54,7 @@ protected:
  * accessing the Compatiblity extension from the Typed Elements Registry is not recommended, such as for MASS.
  */
 UCLASS()
-class TYPEDELEMENTSDATASTORAGE_API UTypedElementDataStorageCompatibilitySubsystem : public UWorldSubsystem
+class TYPEDELEMENTSDATASTORAGE_API UTypedElementDataStorageCompatibilitySubsystem : public UEditorSubsystem
 {
 	GENERATED_BODY()
 
@@ -69,7 +65,5 @@ public:
 	const ITypedElementDataStorageCompatibilityInterface* Get() const;
 
 protected:
-	bool DoesSupportWorldType(const EWorldType::Type WorldType) const override;
-
 	mutable ITypedElementDataStorageCompatibilityInterface* DataStorageCompatibility{ nullptr };
 };
\ No newline at end of file
diff --git a/Engine/Plugins/TypedElementsDataStorage/Source/TypedElementsDataStorage/TypedElementsDataStorage.Build.cs b/Engine/Plugins/TypedElementsDataStorage/Source/TypedElementsDataStorage/TypedElementsDataStorage.Build.cs
--- a/Engine/Plugins/TypedElementsDataStorage/Source/TypedElementsDataStorage/TypedElementsDataStorage.Build.cs
+++ b/Engine/Plugins/TypedElementsDataStorage/Source/TypedElementsDataStorage/TypedElementsDataStorage.Build.cs
@@ -21,8 +21,11 @@ public class TypedElementsDataStorage : ModuleRules
 					"Engine",
 					"MassActors",
 					"MassEntity",
+					"MassEntityEditor",
+					"MassGameplayEditor",
 					"TypedElementFramework",
-					"UnrealEd"
+					"EditorSubsystem",
+					"UnrealEd",
 				});
 
 			PrivateDependencyModuleNames.AddRange(new string[] {});
diff --git a/Engine/Plugins/TypedElementsDataStorage/Source/TypedElementsDataStorageUI/TypedElementsDataStorageUI.Build.cs b/Engine/Plugins/TypedElementsDataStorage/Source/TypedElementsDataStorageUI/TypedElementsDataStorageUI.Build.cs
--- a/Engine/Plugins/TypedElementsDataStorage/Source/TypedElementsDataStorageUI/TypedElementsDataStorageUI.Build.cs
+++ b/Engine/Plugins/TypedElementsDataStorage/Source/TypedElementsDataStorageUI/TypedElementsDataStorageUI.Build.cs
@@ -24,7 +24,9 @@ public class TypedElementsDataStorageUI : ModuleRules
 					"SlateCore",
 					"ToolMenus",
 					"TypedElementFramework",
-					"TypedElementsDataStorage"
+					"TypedElementsDataStorage",
+					"EditorSubsystem",
+					"UnrealEd",
 				});
 			PrivateDependencyModuleNames.AddRange(
 				new string[]
