Introduced a sync tag to the Typed Elements Data Storage and bidirectional syncing of the label on an actor.

diff --git a/Engine/Plugins/TypedElementsDataStorage/Source/TypedElementsDataStorage/Private/Compatibility/TypedElementActorLabelProcessors.cpp b/Engine/Plugins/TypedElementsDataStorage/Source/TypedElementsDataStorage/Private/Compatibility/TypedElementActorLabelProcessors.cpp
new file mode 100644
--- /dev/null
+++ b/Engine/Plugins/TypedElementsDataStorage/Source/TypedElementsDataStorage/Private/Compatibility/TypedElementActorLabelProcessors.cpp
@@ -0,0 +1,109 @@
+// Copyright Epic Games, Inc. All Rights Reserved.
+
+#include "Compatibility/TypedElementActorLabelProcessors.h"
+
+#include "Elements/Columns/TypedElementMiscColumns.h"
+#include "Elements/Columns/TypedElementLabelColumns.h"
+#include "Hash/CityHash.h"
+#include "MassActorSubsystem.h"
+#include "MassExecutionContext.h"
+
+/**
+ * UTypedElementActorLabelToColumnProcessor
+ */
+
+UTypedElementActorLabelToColumnProcessor::UTypedElementActorLabelToColumnProcessor()
+	: Query(*this)
+{
+	ExecutionFlags = static_cast<int32>(EProcessorExecutionFlags::Editor);
+	ExecutionOrder.ExecuteInGroup = UE::Mass::ProcessorGroupNames::SyncWorldToMass;
+	ProcessingPhase = EMassProcessingPhase::PrePhysics;
+	bRequiresGameThreadExecution = true;
+}
+
+void UTypedElementActorLabelToColumnProcessor::ConfigureQueries()
+{
+	Query.AddRequirement(FMassActorFragment::StaticStruct(), EMassFragmentAccess::ReadOnly);
+	Query.AddRequirement(FTypedElementLabelColumn::StaticStruct(), EMassFragmentAccess::ReadWrite);
+	Query.AddRequirement(FTypedElementLabelHashColumn::StaticStruct(), EMassFragmentAccess::ReadWrite);
+}
+
+void UTypedElementActorLabelToColumnProcessor::Execute(FMassEntityManager& EntityManager, FMassExecutionContext& Context)
+{
+	Query.ForEachEntityChunk(EntityManager, Context, [](FMassExecutionContext& Context)
+		{
+			const FMassActorFragment* ActorIt = Context.GetFragmentView<FMassActorFragment>().GetData();
+			FTypedElementLabelColumn* LabelIt = Context.GetMutableFragmentView<FTypedElementLabelColumn>().GetData();
+			FTypedElementLabelHashColumn* LabelHashIt = Context.GetMutableFragmentView<FTypedElementLabelHashColumn>().GetData();
+			
+			const int32 NumEntities = Context.GetNumEntities();
+			for (int32 Index = 0; Index < NumEntities; ++Index)
+			{
+				const FString& ActorLabel = ActorIt->Get()->GetActorLabel(false);
+				uint64 ActorLabelHash = CityHash64(
+					reinterpret_cast<const char*>(ActorLabel.GetCharArray().GetData()), 
+					ActorLabel.Len() * sizeof(ActorLabel.GetCharArray().GetData()[0]));
+				if (LabelHashIt->LabelHash != ActorLabelHash)
+				{
+					LabelIt->Label = ActorLabel;
+					LabelHashIt->LabelHash = ActorLabelHash;
+				}
+				
+				++ActorIt;
+				++LabelIt;
+				++LabelHashIt;
+			}
+		}
+	);
+}
+
+
+
+/**
+ * UTypedElementLabelColumnToActorProcessor
+ */
+
+UTypedElementLabelColumnToActorProcessor::UTypedElementLabelColumnToActorProcessor()
+	: Query(*this)
+{
+	ExecutionFlags = static_cast<int32>(EProcessorExecutionFlags::Editor);
+	ExecutionOrder.ExecuteInGroup = UE::Mass::ProcessorGroupNames::UpdateWorldFromMass;
+	ProcessingPhase = EMassProcessingPhase::FrameEnd;
+	bRequiresGameThreadExecution = true;
+}
+
+void UTypedElementLabelColumnToActorProcessor::ConfigureQueries()
+{
+	Query.AddRequirement(FMassActorFragment::StaticStruct(), EMassFragmentAccess::ReadWrite);
+	Query.AddRequirement(FTypedElementLabelColumn::StaticStruct(), EMassFragmentAccess::ReadOnly);
+	Query.AddRequirement(FTypedElementLabelHashColumn::StaticStruct(), EMassFragmentAccess::ReadOnly);
+	Query.AddTagRequirement(*FTypedElementSyncBackToWorldTag::StaticStruct(), EMassFragmentPresence::All);
+}
+
+void UTypedElementLabelColumnToActorProcessor::Execute(FMassEntityManager& EntityManager, FMassExecutionContext& Context)
+{
+	Query.ForEachEntityChunk(EntityManager, Context, [](FMassExecutionContext& Context)
+		{
+			FMassActorFragment* ActorIt = Context.GetMutableFragmentView<FMassActorFragment>().GetData();
+			const FTypedElementLabelColumn* LabelIt = Context.GetFragmentView<FTypedElementLabelColumn>().GetData();
+			const FTypedElementLabelHashColumn* LabelHashIt = Context.GetFragmentView<FTypedElementLabelHashColumn>().GetData();
+
+			const int32 NumEntities = Context.GetNumEntities();
+			for (int32 Index = 0; Index < NumEntities; ++Index)
+			{
+				const FString& ActorLabel = ActorIt->Get()->GetActorLabel(false);
+				uint64 ActorLabelHash = CityHash64(
+					reinterpret_cast<const char*>(ActorLabel.GetCharArray().GetData()),
+					ActorLabel.Len() * sizeof(ActorLabel.GetCharArray().GetData()[0]));
+				if (LabelHashIt->LabelHash != ActorLabelHash)
+				{
+					ActorIt->GetMutable()->SetActorLabel(LabelIt->Label);
+				}
+				
+				++ActorIt;
+				++LabelIt;
+				++LabelHashIt;
+			}
+		}
+	);
+}
diff --git a/Engine/Plugins/TypedElementsDataStorage/Source/TypedElementsDataStorage/Private/Compatibility/TypedElementActorLabelProcessors.h b/Engine/Plugins/TypedElementsDataStorage/Source/TypedElementsDataStorage/Private/Compatibility/TypedElementActorLabelProcessors.h
new file mode 100644
--- /dev/null
+++ b/Engine/Plugins/TypedElementsDataStorage/Source/TypedElementsDataStorage/Private/Compatibility/TypedElementActorLabelProcessors.h
@@ -0,0 +1,49 @@
+// Copyright Epic Games, Inc. All Rights Reserved.
+
+#pragma once
+
+#include "MassProcessor.h"
+#include "UObject/ObjectMacros.h"
+
+#include "TypedElementActorLabelProcessors.generated.h"
+
+/**
+ * Takes the label set on an actor and copies it to the Data Storage if they differ.
+ */
+UCLASS()
+class UTypedElementActorLabelToColumnProcessor : public UMassProcessor
+{
+	GENERATED_BODY()
+
+public:
+	UTypedElementActorLabelToColumnProcessor();
+
+protected:
+	void ConfigureQueries() override;
+	void Execute(FMassEntityManager& EntityManager, FMassExecutionContext& Context) override;
+
+private:
+	FMassEntityQuery Query;
+};
+
+
+
+/**
+ * Takes the label stored in the Data Storage and copies it to the actor's label if the FTypedElementSyncBackToWorldTag
+ * has been set and the labels differ.
+ */
+UCLASS()
+class UTypedElementLabelColumnToActorProcessor : public UMassProcessor
+{
+	GENERATED_BODY()
+
+public:
+	UTypedElementLabelColumnToActorProcessor();
+
+protected:
+	void ConfigureQueries() override;
+	void Execute(FMassEntityManager& EntityManager, FMassExecutionContext& Context) override;
+
+private:
+	FMassEntityQuery Query;
+};
\ No newline at end of file
diff --git a/Engine/Plugins/TypedElementsDataStorage/Source/TypedElementsDataStorage/Private/Processors/TypedElementMiscProcessors.cpp b/Engine/Plugins/TypedElementsDataStorage/Source/TypedElementsDataStorage/Private/Processors/TypedElementMiscProcessors.cpp
new file mode 100644
--- /dev/null
+++ b/Engine/Plugins/TypedElementsDataStorage/Source/TypedElementsDataStorage/Private/Processors/TypedElementMiscProcessors.cpp
@@ -0,0 +1,36 @@
+// Copyright Epic Games, Inc. All Rights Reserved.
+
+#include "Processors/TypedElementMiscProcessors.h"
+
+#include "MassCommonTypes.h"
+#include "MassExecutionContext.h"
+#include "Elements/Columns/TypedElementMiscColumns.h"
+
+/**
+ * UTypedElementRemoveSyncToWorldTagProcessor
+ */
+
+UTypedElementRemoveSyncToWorldTagProcessor::UTypedElementRemoveSyncToWorldTagProcessor()
+	: Query(*this)
+{
+	ExecutionFlags = static_cast<int32>(EProcessorExecutionFlags::Editor);
+	ExecutionOrder.ExecuteAfter.Add(UE::Mass::ProcessorGroupNames::UpdateWorldFromMass);
+	ProcessingPhase = EMassProcessingPhase::FrameEnd;
+}
+
+void UTypedElementRemoveSyncToWorldTagProcessor::ConfigureQueries()
+{
+	Query.AddTagRequirement(*FTypedElementSyncBackToWorldTag::StaticStruct(), EMassFragmentPresence::All);
+}
+
+void UTypedElementRemoveSyncToWorldTagProcessor::Execute(FMassEntityManager& EntityManager, FMassExecutionContext& Context)
+{
+	Query.ForEachEntityChunk(EntityManager, Context, [](FMassExecutionContext& Context)
+		{
+			for (FMassEntityHandle Entity : Context.GetEntities())
+			{
+				Context.Defer().RemoveTag_RuntimeCheck<FTypedElementSyncBackToWorldTag>(Entity);
+			}
+		}
+	);
+}
diff --git a/Engine/Plugins/TypedElementsDataStorage/Source/TypedElementsDataStorage/Private/Processors/TypedElementMiscProcessors.h b/Engine/Plugins/TypedElementsDataStorage/Source/TypedElementsDataStorage/Private/Processors/TypedElementMiscProcessors.h
new file mode 100644
--- /dev/null
+++ b/Engine/Plugins/TypedElementsDataStorage/Source/TypedElementsDataStorage/Private/Processors/TypedElementMiscProcessors.h
@@ -0,0 +1,27 @@
+// Copyright Epic Games, Inc. All Rights Reserved.
+
+#pragma once
+
+#include "MassProcessor.h"
+#include "UObject/ObjectMacros.h"
+
+#include "TypedElementMiscProcessors.generated.h"
+
+/**
+ * Removes all FTypedElementSyncBackToWorldTags at the end of an update cycle.
+ */
+UCLASS()
+class UTypedElementRemoveSyncToWorldTagProcessor : public UMassProcessor
+{
+	GENERATED_BODY()
+
+public:
+	UTypedElementRemoveSyncToWorldTagProcessor();
+
+protected:
+	void ConfigureQueries() override;
+	void Execute(FMassEntityManager& EntityManager, FMassExecutionContext& Context) override;
+
+private:
+	FMassEntityQuery Query;
+};
\ No newline at end of file
diff --git a/Engine/Source/Runtime/TypedElementFramework/Public/Elements/Columns/TypedElementLabelColumns.h b/Engine/Source/Runtime/TypedElementFramework/Public/Elements/Columns/TypedElementLabelColumns.h
new file mode 100644
--- /dev/null
+++ b/Engine/Source/Runtime/TypedElementFramework/Public/Elements/Columns/TypedElementLabelColumns.h
@@ -0,0 +1,33 @@
+// Copyright Epic Games, Inc. All Rights Reserved.
+
+#pragma once
+
+#include "Containers/UnrealString.h"
+#include "Elements/Interfaces/TypedElementDataStorageInterface.h"
+#include "UObject/ObjectMacros.h"
+
+#include "TypedElementLabelColumns.generated.h"
+
+
+/**
+ * Column that stores a label.
+ */
+USTRUCT()
+struct FTypedElementLabelColumn : public FTypedElementDataStorageColumn
+{
+	GENERATED_BODY()
+
+	FString Label;
+};
+
+/**
+ * Column that stores the hash of a label. This is typically paired with FTypedElementLabelColumn, but 
+ * kept separate in order to iterate quickly over all hash values.
+ */
+USTRUCT()
+struct FTypedElementLabelHashColumn : public FTypedElementDataStorageColumn
+{
+	GENERATED_BODY()
+
+	uint64 LabelHash;
+};
\ No newline at end of file
diff --git a/Engine/Source/Runtime/TypedElementFramework/Public/Elements/Columns/TypedElementMiscColumns.h b/Engine/Source/Runtime/TypedElementFramework/Public/Elements/Columns/TypedElementMiscColumns.h
new file mode 100644
--- /dev/null
+++ b/Engine/Source/Runtime/TypedElementFramework/Public/Elements/Columns/TypedElementMiscColumns.h
@@ -0,0 +1,20 @@
+// Copyright Epic Games, Inc. All Rights Reserved.
+
+#pragma once
+
+#include "Elements/Interfaces/TypedElementDataStorageInterface.h"
+#include "UObject/ObjectMacros.h"
+
+#include "TypedElementMiscColumns.generated.h"
+
+
+/**
+ * Tag to indicate that there are one or more bits of information in the row that
+ * need to be copied out the Data Storage and into the original object. This tag
+ * will automatically be removed at the end of a tick.
+ */
+USTRUCT()
+struct FTypedElementSyncBackToWorldTag : public FTypedElementDataStorageTag
+{
+	GENERATED_BODY()
+};
