Added bi-directional support for syncing transform information between Actors and the Typed Elements Data Storage.

diff --git a/Engine/Plugins/TypedElementsDataStorage/Source/TypedElementsDataStorage/Private/Compatibility/TypedElementActorTransformProcessors.cpp b/Engine/Plugins/TypedElementsDataStorage/Source/TypedElementsDataStorage/Private/Compatibility/TypedElementActorTransformProcessors.cpp
new file mode 100644
--- /dev/null
+++ b/Engine/Plugins/TypedElementsDataStorage/Source/TypedElementsDataStorage/Private/Compatibility/TypedElementActorTransformProcessors.cpp
@@ -0,0 +1,132 @@
+// Copyright Epic Games, Inc. All Rights Reserved.
+
+#include "Compatibility/TypedElementActorTransformProcessors.h"
+
+#include "Elements/Columns/TypedElementMiscColumns.h"
+#include "Elements/Columns/TypedElementTransformColumns.h"
+#include "MassActorSubsystem.h"
+#include "MassExecutionContext.h"
+
+/**
+ * UTypedElementActorAddTransformColumnProcessor
+ */
+
+UTypedElementActorAddTransformColumnProcessor::UTypedElementActorAddTransformColumnProcessor()
+	: Query(*this)
+{
+	ExecutionFlags = static_cast<int32>(EProcessorExecutionFlags::Editor);
+	// Run this before the world gets synced to MASS so a new transform is initialized in the same frame.
+	ExecutionOrder.ExecuteBefore.Add(UE::Mass::ProcessorGroupNames::SyncWorldToMass);
+	ProcessingPhase = EMassProcessingPhase::PrePhysics;
+	bRequiresGameThreadExecution = true;
+}
+
+void UTypedElementActorAddTransformColumnProcessor::ConfigureQueries()
+{
+	Query.AddRequirement(FMassActorFragment::StaticStruct(), EMassFragmentAccess::ReadOnly);
+	Query.AddRequirement(FTypedElementLocalTransformColumn::StaticStruct(), EMassFragmentAccess::ReadOnly, EMassFragmentPresence::None);
+}
+
+void UTypedElementActorAddTransformColumnProcessor::Execute(FMassEntityManager& EntityManager, FMassExecutionContext& Context)
+{
+	Query.ForEachEntityChunk(EntityManager, Context, [](FMassExecutionContext& Context)
+		{
+			const FMassActorFragment* ActorIt = Context.GetFragmentView<FMassActorFragment>().GetData();
+			
+			for (FMassEntityHandle Entity : Context.GetEntities())
+			{
+				if (ActorIt->Get()->GetRootComponent())
+				{
+					Context.Defer().AddFragment_RuntimeCheck<FTypedElementLocalTransformColumn>(Entity);
+				}
+				
+				++ActorIt;
+			}
+		}
+	);
+}
+
+
+/**
+ * UTypedElementActorLocalTransformToColumnProcessor
+ */
+
+UTypedElementActorLocalTransformToColumnProcessor::UTypedElementActorLocalTransformToColumnProcessor()
+	: Query(*this)
+{
+	ExecutionFlags = static_cast<int32>(EProcessorExecutionFlags::Editor);
+	ExecutionOrder.ExecuteInGroup = UE::Mass::ProcessorGroupNames::SyncWorldToMass;
+	ProcessingPhase = EMassProcessingPhase::PrePhysics;
+	bRequiresGameThreadExecution = true;
+}
+
+void UTypedElementActorLocalTransformToColumnProcessor::ConfigureQueries()
+{
+	Query.AddRequirement(FMassActorFragment::StaticStruct(), EMassFragmentAccess::ReadOnly);
+	Query.AddRequirement(FTypedElementLocalTransformColumn::StaticStruct(), EMassFragmentAccess::ReadWrite);
+}
+
+void UTypedElementActorLocalTransformToColumnProcessor::Execute(FMassEntityManager& EntityManager, FMassExecutionContext& Context)
+{
+	Query.ForEachEntityChunk(EntityManager, Context, [](FMassExecutionContext& Context)
+		{
+			const FMassActorFragment* ActorIt = Context.GetFragmentView<FMassActorFragment>().GetData();
+			FTypedElementLocalTransformColumn* TransformIt = Context.GetMutableFragmentView<FTypedElementLocalTransformColumn>().GetData();
+
+			for (FMassEntityHandle Entity : Context.GetEntities())
+			{
+				if (ActorIt->Get()->GetRootComponent() != nullptr)
+				{
+					TransformIt->Transform = ActorIt->Get()->GetActorTransform();
+				}
+				else
+				{
+					Context.Defer().RemoveFragment_RuntimeCheck<FTypedElementLocalTransformColumn>(Entity);
+				}
+				
+				++ActorIt;
+				++TransformIt;
+			}
+		}
+	);
+}
+
+
+/**
+ * UTypedElementTransformColumnToActorProcessor
+ */
+
+UTypedElementTransformColumnToActorProcessor::UTypedElementTransformColumnToActorProcessor()
+	: Query(*this)
+{
+	ExecutionFlags = static_cast<int32>(EProcessorExecutionFlags::Editor);
+	ExecutionOrder.ExecuteInGroup = UE::Mass::ProcessorGroupNames::UpdateWorldFromMass;
+	ProcessingPhase = EMassProcessingPhase::FrameEnd;
+	bRequiresGameThreadExecution = true;
+}
+
+void UTypedElementTransformColumnToActorProcessor::ConfigureQueries()
+{
+	Query.AddRequirement(FMassActorFragment::StaticStruct(), EMassFragmentAccess::ReadWrite);
+	Query.AddRequirement(FTypedElementLocalTransformColumn::StaticStruct(), EMassFragmentAccess::ReadOnly);
+	Query.AddTagRequirement(*FTypedElementSyncBackToWorldTag::StaticStruct(), EMassFragmentPresence::All);
+}
+
+void UTypedElementTransformColumnToActorProcessor::Execute(FMassEntityManager& EntityManager, FMassExecutionContext& Context)
+{
+	Query.ForEachEntityChunk(EntityManager, Context, [](FMassExecutionContext& Context)
+		{
+			FMassActorFragment* ActorIt = Context.GetMutableFragmentView<FMassActorFragment>().GetData();
+			const FTypedElementLocalTransformColumn* TransformIt = Context.GetFragmentView<FTypedElementLocalTransformColumn>().GetData();
+			
+			const int32 NumEntities = Context.GetNumEntities();
+			for (int32 Index = 0; Index < NumEntities; ++Index)
+			{
+				ActorIt->GetMutable()->SetActorTransform(TransformIt->Transform);
+				
+				++ActorIt;
+				++TransformIt;
+			}
+		}
+	);
+}
diff --git a/Engine/Plugins/TypedElementsDataStorage/Source/TypedElementsDataStorage/Private/Compatibility/TypedElementActorTransformProcessors.h b/Engine/Plugins/TypedElementsDataStorage/Source/TypedElementsDataStorage/Private/Compatibility/TypedElementActorTransformProcessors.h
new file mode 100644
--- /dev/null
+++ b/Engine/Plugins/TypedElementsDataStorage/Source/TypedElementsDataStorage/Private/Compatibility/TypedElementActorTransformProcessors.h
@@ -0,0 +1,71 @@
+// Copyright Epic Games, Inc. All Rights Reserved.
+
+#pragma once
+
+#include "MassProcessor.h"
+#include "UObject/ObjectMacros.h"
+
+#include "TypedElementActorTransformProcessors.generated.h"
+
+/**
+ * Checks actors that don't have a tranform column and adds one if an actor has been
+ * assigned a transform.
+ */
+UCLASS()
+class UTypedElementActorAddTransformColumnProcessor : public UMassProcessor
+{
+	GENERATED_BODY()
+
+public:
+	UTypedElementActorAddTransformColumnProcessor();
+
+protected:
+	void ConfigureQueries() override;
+	void Execute(FMassEntityManager& EntityManager, FMassExecutionContext& Context) override;
+
+private:
+	FMassEntityQuery Query;
+};
+
+
+/**
+ * Takes the transform set on an actor and copies it to the Data Storage or removes the 
+ * transform column if there's not transform available anymore.
+ */
+UCLASS()
+class UTypedElementActorLocalTransformToColumnProcessor : public UMassProcessor
+{
+	GENERATED_BODY()
+
+public:
+	UTypedElementActorLocalTransformToColumnProcessor();
+
+protected:
+	void ConfigureQueries() override;
+	void Execute(FMassEntityManager& EntityManager, FMassExecutionContext& Context) override;
+
+private:
+	FMassEntityQuery Query;
+};
+
+
+
+/**
+ * Takes the transform stored in the Data Storage and copies it to the actor's tranform if 
+ * the FTypedElementSyncBackToWorldTag has been set.
+ */
+UCLASS()
+class UTypedElementTransformColumnToActorProcessor : public UMassProcessor
+{
+	GENERATED_BODY()
+
+public:
+	UTypedElementTransformColumnToActorProcessor();
+
+protected:
+	void ConfigureQueries() override;
+	void Execute(FMassEntityManager& EntityManager, FMassExecutionContext& Context) override;
+
+private:
+	FMassEntityQuery Query;
+};
\ No newline at end of file
diff --git a/Engine/Source/Runtime/TypedElementFramework/Public/Elements/Columns/TypedElementTransformColumns.h b/Engine/Source/Runtime/TypedElementFramework/Public/Elements/Columns/TypedElementTransformColumns.h
new file mode 100644
--- /dev/null
+++ b/Engine/Source/Runtime/TypedElementFramework/Public/Elements/Columns/TypedElementTransformColumns.h
@@ -0,0 +1,25 @@
+// Copyright Epic Games, Inc. All Rights Reserved.
+
+#pragma once
+
+#include "Elements/Interfaces/TypedElementDataStorageInterface.h"
+#include "Math/MathFwd.h"
+#include "UObject/ObjectMacros.h"
+
+#include "TypedElementTransformColumns.generated.h"
+
+
+/**
+ * Column that stores a local transform. 
+ */
+USTRUCT()
+struct FTypedElementLocalTransformColumn : public FTypedElementDataStorageColumn
+{
+	GENERATED_BODY()
+
+	// Transform is not being intialized to avoid spending time on initialization when the
+	// Transfrom will be updated the first and following ticks after it's creation. If this
+	// isn't initialized at the correct time, then the sync from source or the true initialization
+	// need to be moved to an earlier phase or group.
+	FTransform Transform { NoInit };
+};
