Synchronized data to/from actors and the Typed Element Data Storage.

This change adds uni-directional synchronization of package information from actors to the Typed Element Data Storage. It also extends the definition of the actor table to include the recently added columns.

diff --git a/Engine/Plugins/TypedElementsDataStorage/Source/TypedElementsDataStorage/Private/Compatibility/TypedElementActorPackagePathToColumnProcessor.cpp b/Engine/Plugins/TypedElementsDataStorage/Source/TypedElementsDataStorage/Private/Compatibility/TypedElementActorPackagePathToColumnProcessor.cpp
new file mode 100644
--- /dev/null
+++ b/Engine/Plugins/TypedElementsDataStorage/Source/TypedElementsDataStorage/Private/Compatibility/TypedElementActorPackagePathToColumnProcessor.cpp
@@ -0,0 +1,50 @@
+// Copyright Epic Games, Inc. All Rights Reserved.
+
+#include "Compatibility/TypedElementActorPackagePathToColumnProcessor.h"
+
+#include "Elements/Columns/TypedElementPackageColumns.h"
+#include "MassActorSubsystem.h"
+#include "MassExecutionContext.h"
+#include "TypedElementSubsystems.h"
+
+UTypedElementActorPackagePathToColumnProcessor::UTypedElementActorPackagePathToColumnProcessor()
+	: Query(*this)
+{
+	ExecutionFlags = static_cast<int32>(EProcessorExecutionFlags::Editor);
+	ObservedType = FMassActorFragment::StaticStruct();
+	Operation = EMassObservedOperation::Add;
+	bRequiresGameThreadExecution = true;
+}
+
+void UTypedElementActorPackagePathToColumnProcessor::ConfigureQueries()
+{
+	Query.AddRequirement(FMassActorFragment::StaticStruct(), EMassFragmentAccess::ReadOnly);
+	Query.AddRequirement(FTypedElementPackagePathColumn::StaticStruct(), EMassFragmentAccess::ReadWrite);
+	Query.AddRequirement(FTypedElementPackageLoadedPathColumn::StaticStruct(), EMassFragmentAccess::ReadWrite);
+	
+	ProcessorRequirements.AddSubsystemRequirement(UTypedElementDataStorageSubsystem::StaticClass(), EMassFragmentAccess::ReadWrite);
+}
+
+void UTypedElementActorPackagePathToColumnProcessor::Execute(FMassEntityManager& EntityManager, FMassExecutionContext& Context)
+{
+	Query.ForEachEntityChunk(EntityManager, Context, [](FMassExecutionContext& Context)
+		{
+			const FMassActorFragment* ActorIt = Context.GetFragmentView<FMassActorFragment>().GetData();
+			FTypedElementPackagePathColumn* PathIt = Context.GetMutableFragmentView<FTypedElementPackagePathColumn>().GetData();
+			FTypedElementPackageLoadedPathColumn* LoadedPathIt = Context.GetMutableFragmentView<FTypedElementPackageLoadedPathColumn>().GetData();
+			
+			const int32 EntityCount = Context.GetNumEntities();
+			for (int32 Counter=0; Counter < EntityCount; ++Counter)
+			{
+				UPackage* Target = ActorIt->Get()->GetPackage();
+
+				Target->GetPathName(nullptr, PathIt->Path);
+				LoadedPathIt->LoadedPath = Target->GetLoadedPath();
+				
+				++ActorIt;
+				++PathIt;
+				++LoadedPathIt;
+			}
+		}
+	);
+}
diff --git a/Engine/Plugins/TypedElementsDataStorage/Source/TypedElementsDataStorage/Private/Compatibility/TypedElementActorPackagePathToColumnProcessor.h b/Engine/Plugins/TypedElementsDataStorage/Source/TypedElementsDataStorage/Private/Compatibility/TypedElementActorPackagePathToColumnProcessor.h
new file mode 100644
--- /dev/null
+++ b/Engine/Plugins/TypedElementsDataStorage/Source/TypedElementsDataStorage/Private/Compatibility/TypedElementActorPackagePathToColumnProcessor.h
@@ -0,0 +1,24 @@
+// Copyright Epic Games, Inc. All Rights Reserved.
+
+#pragma once
+
+#include "MassObserverProcessor.h"
+#include "UObject/ObjectMacros.h"
+
+#include "TypedElementActorPackagePathToColumnProcessor.generated.h"
+
+UCLASS()
+class UTypedElementActorPackagePathToColumnProcessor : public UMassObserverProcessor
+{
+	GENERATED_BODY()
+
+public:
+	UTypedElementActorPackagePathToColumnProcessor();
+
+protected:
+	void ConfigureQueries() override;
+	void Execute(FMassEntityManager& EntityManager, FMassExecutionContext& Context) override;
+
+private:
+	FMassEntityQuery Query;
+};
diff --git a/Engine/Plugins/TypedElementsDataStorage/Source/TypedElementsDataStorage/Private/TypedElementDatabaseCompatibility.cpp b/Engine/Plugins/TypedElementsDataStorage/Source/TypedElementsDataStorage/Private/TypedElementDatabaseCompatibility.cpp
--- a/Engine/Plugins/TypedElementsDataStorage/Source/TypedElementsDataStorage/Private/TypedElementDatabaseCompatibility.cpp
+++ b/Engine/Plugins/TypedElementsDataStorage/Source/TypedElementsDataStorage/Private/TypedElementDatabaseCompatibility.cpp
@@ -3,6 +3,10 @@
 #include "TypedElementDatabaseCompatibility.h"
 
 #include "Editor.h"
+#include "Elements/Columns/TypedElementLabelColumns.h"
+#include "Elements/Columns/TypedElementMiscColumns.h"
+#include "Elements/Columns/TypedElementPackageColumns.h"
+#include "Elements/Columns/TypedElementTransformColumns.h"
 #include "MassActorEditorSubsystem.h"
 #include "MassActorSubsystem.h"
 
@@ -72,8 +76,19 @@ void UTypedElementDatabaseCompatibility::CreateStandardArchetypes()
 {
 	StandardActorTable = Storage->RegisterTable(MakeArrayView(
 		{
-			FMassActorFragment::StaticStruct()
+			FMassActorFragment::StaticStruct(),
+			FTypedElementLabelColumn::StaticStruct(),
+			FTypedElementLabelHashColumn::StaticStruct(),
+			FTypedElementPackagePathColumn::StaticStruct(),
+			FTypedElementPackageLoadedPathColumn::StaticStruct()
 		}), FName("Editor_StandardActorTable"));
+
+	StandardActorWithTransformTable = Storage->RegisterTable(StandardActorTable,
+		MakeArrayView(
+		{
+			FTypedElementLabelColumn::StaticStruct(),
+			FTypedElementLocalTransformColumn::StaticStruct()
+		}), FName("Editor_StandardActorWithTransformTable"));
 }
 
 void UTypedElementDatabaseCompatibility::AddPendingActors()
diff --git a/Engine/Plugins/TypedElementsDataStorage/Source/TypedElementsDataStorage/Private/TypedElementDatabaseCompatibility.h b/Engine/Plugins/TypedElementsDataStorage/Source/TypedElementsDataStorage/Private/TypedElementDatabaseCompatibility.h
--- a/Engine/Plugins/TypedElementsDataStorage/Source/TypedElementsDataStorage/Private/TypedElementDatabaseCompatibility.h
+++ b/Engine/Plugins/TypedElementsDataStorage/Source/TypedElementsDataStorage/Private/TypedElementDatabaseCompatibility.h
@@ -39,6 +39,7 @@ private:
 	TArray<TWeakObjectPtr<AActor>> ActorsPendingRegistration;
 	
 	TypedElementTableHandle StandardActorTable{ TypedElementInvalidTableHandle };
+	TypedElementTableHandle StandardActorWithTransformTable{ TypedElementInvalidTableHandle };
 	ITypedElementDataStorageInterface* Storage{ nullptr };
 	TSharedPtr<FMassActorManager> ActorSubsystem;
 };
\ No newline at end of file
diff --git a/Engine/Source/Runtime/TypedElementFramework/Public/Elements/Columns/TypedElementPackageColumns.h b/Engine/Source/Runtime/TypedElementFramework/Public/Elements/Columns/TypedElementPackageColumns.h
new file mode 100644
--- /dev/null
+++ b/Engine/Source/Runtime/TypedElementFramework/Public/Elements/Columns/TypedElementPackageColumns.h
@@ -0,0 +1,39 @@
+// Copyright Epic Games, Inc. All Rights Reserved.
+
+#pragma once
+
+#include "Containers/UnrealString.h"
+#include "Elements/Interfaces/TypedElementDataStorageInterface.h"
+#include "IO/PackageId.h"
+#include "Misc/PackagePath.h"
+#include "UObject/ObjectMacros.h"
+
+#include "TypedElementPackageColumns.generated.h"
+
+
+/**
+ * Column that stores the path of a package.
+ */
+USTRUCT()
+struct FTypedElementPackagePathColumn : public FTypedElementDataStorageColumn
+{
+	GENERATED_BODY()
+
+	FString Path;
+};
+
+inline uint32 GetTypeHash(const FTypedElementPackagePathColumn& InStruct)
+{
+	return GetTypeHash(InStruct.Path);
+}
+
+/**
+ * Column that stores the full loading path to a package.
+ */
+USTRUCT()
+struct FTypedElementPackageLoadedPathColumn : public FTypedElementDataStorageColumn
+{
+	GENERATED_BODY()
+
+	FPackagePath LoadedPath;
+};
