Fixed a crash in UTypedElementActorLocalTransformToColumnProcessor::Execute ocurring while closing down the editor, when used weak object pointers are no longer valid

diff --git a/Engine/Plugins/TypedElementsDataStorage/Source/TypedElementsDataStorage/Private/Compatibility/TypedElementActorTransformProcessors.cpp b/Engine/Plugins/TypedElementsDataStorage/Source/TypedElementsDataStorage/Private/Compatibility/TypedElementActorTransformProcessors.cpp
--- a/Engine/Plugins/TypedElementsDataStorage/Source/TypedElementsDataStorage/Private/Compatibility/TypedElementActorTransformProcessors.cpp
+++ b/Engine/Plugins/TypedElementsDataStorage/Source/TypedElementsDataStorage/Private/Compatibility/TypedElementActorTransformProcessors.cpp
@@ -75,7 +75,8 @@ void UTypedElementActorLocalTransformToColumnProcessor::Execute(FMassEntityManag
 
 			for (FMassEntityHandle Entity : Context.GetEntities())
 			{
-				if (ActorIt->Get()->GetRootComponent() != nullptr)
+				const AActor* ActorInstance = ActorIt->Get();
+				if (ActorInstance != nullptr && ActorInstance->GetRootComponent() != nullptr)
 				{
 					TransformIt->Transform = ActorIt->Get()->GetActorTransform();
 				}
