Fix compile errors
#rnx

diff --git a/Engine/Plugins/TypedElementsDataStorage/Source/TypedElementsDataStorage/Private/TypedElementDatabase.cpp b/Engine/Plugins/TypedElementsDataStorage/Source/TypedElementsDataStorage/Private/TypedElementDatabase.cpp
--- a/Engine/Plugins/TypedElementsDataStorage/Source/TypedElementsDataStorage/Private/TypedElementDatabase.cpp
+++ b/Engine/Plugins/TypedElementsDataStorage/Source/TypedElementsDataStorage/Private/TypedElementDatabase.cpp
@@ -175,7 +175,7 @@ void* UTypedElementDatabase::AddOrGetColumnData(TypedElementRowHandle Row, const
 			checkf(Column.IsValid(), TEXT("Added a new column to the Typed Element's data storae, but it couldn't be retrieved."));
 
 		}
-		return Column.GetMutableMemory();
+		return Column.GetMemory();
 	}
 	return nullptr;
 }
@@ -215,7 +215,7 @@ void* UTypedElementDatabase::GetColumnData(TypedElementRowHandle Row, const UScr
 		FStructView Column = ActiveEditorEntityManager->GetFragmentDataStruct(Entity, ColumnType);
 		if (Column.IsValid())
 		{
-			return Column.GetMutableMemory();
+			return Column.GetMemory();
 		}
 	}
 	return nullptr;
@@ -242,7 +242,7 @@ ColumnDataResult UTypedElementDatabase::GetColumnData(TypedElementRowHandle Row,
 			FStructView Column = ActiveEditorEntityManager->GetFragmentDataStruct(Entity, FragmentType);
 			if (Column.IsValid())
 			{
-				return ColumnDataResult{ FragmentType, Column.GetMutableMemory() };
+				return ColumnDataResult{ FragmentType, Column.GetMemory() };
 			}
 		}
 	}
