Linux editor build fix - added missing StructUtils dependency to TypedElementsDataStorage module

diff --git a/Engine/Plugins/TypedElementsDataStorage/Source/TypedElementsDataStorage/TypedElementsDataStorage.Build.cs b/Engine/Plugins/TypedElementsDataStorage/Source/TypedElementsDataStorage/TypedElementsDataStorage.Build.cs
--- a/Engine/Plugins/TypedElementsDataStorage/Source/TypedElementsDataStorage/TypedElementsDataStorage.Build.cs
+++ b/Engine/Plugins/TypedElementsDataStorage/Source/TypedElementsDataStorage/TypedElementsDataStorage.Build.cs
@@ -26,6 +26,7 @@ public class TypedElementsDataStorage : ModuleRules
 					"TypedElementFramework",
 					"EditorSubsystem",
 					"UnrealEd",
+					"StructUtils"
 				});
 
 			PrivateDependencyModuleNames.AddRange(new string[] {});
diff --git a/Engine/Plugins/TypedElementsDataStorage/TypedElementsDataStorage.uplugin b/Engine/Plugins/TypedElementsDataStorage/TypedElementsDataStorage.uplugin
--- a/Engine/Plugins/TypedElementsDataStorage/TypedElementsDataStorage.uplugin
+++ b/Engine/Plugins/TypedElementsDataStorage/TypedElementsDataStorage.uplugin
@@ -32,6 +32,10 @@
 		{
 			"Name": "MassGameplay",
 			"Enabled": true
+		},
+		{
+			"Name": "StructUtils",
+			"Enabled": true
 		}
 	]
 }
\ No newline at end of file
