Linux editor build fix - added missing StructUtils dependency to TypedElementsDataStorageUI module

diff --git a/Engine/Plugins/TypedElementsDataStorage/Source/TypedElementsDataStorageUI/TypedElementsDataStorageUI.Build.cs b/Engine/Plugins/TypedElementsDataStorage/Source/TypedElementsDataStorageUI/TypedElementsDataStorageUI.Build.cs
--- a/Engine/Plugins/TypedElementsDataStorage/Source/TypedElementsDataStorageUI/TypedElementsDataStorageUI.Build.cs
+++ b/Engine/Plugins/TypedElementsDataStorage/Source/TypedElementsDataStorageUI/TypedElementsDataStorageUI.Build.cs
@@ -32,7 +32,8 @@ public class TypedElementsDataStorageUI : ModuleRules
 				new string[]
 				{
 					"MassActors",
-					"MassEntity"
+					"MassEntity",
+					"StructUtils"
 				});
 
 			DynamicallyLoadedModuleNames.AddRange(new string[] {});
