Moved the Typed Element Data Storage plugin to the experimental folder.

The Typed Element Data Storage plugin is an experimental plugin that's sill in development. As such it should have been placed in the plugins/experimental folder instead of the plugin/ folder. This CR just moves the plugin into the correct folder as everything else was already correctly setup.

diff --git a/Engine/Plugins/Experimental/TypedElementsDataStorage/Source/TypedElementsDataStorage/Private/Compatibility/TypedElementActorLabelProcessors.cpp b/Engine/Plugins/Experimental/TypedElementsDataStorage/Source/TypedElementsDataStorage/Private/Compatibility/TypedElementActorLabelProcessors.cpp
new file mode 100644
--- /dev/null
+++ b/Engine/Plugins/Experimental/TypedElementsDataStorage/Source/TypedElementsDataStorage/Private/Compatibility/TypedElementActorLabelProcessors.cpp
@@ -0,0 +1,109 @@
+// Copyright Epic Games, Inc. All Rights Reserved.
+
+#include "Compatibility/TypedElementActorLabelProcessors.h"
+
+#include "Elements/Columns/TypedElementMiscColumns.h"
+#include "Elements/Columns/TypedElementLabelColumns.h"
+#include "Hash/CityHash.h"
+#include "MassActorSubsystem.h"
+#include "MassExecutionContext.h"
+
+/**
+ * UTypedElementActorLabelToColumnProcessor
+ */
+
+UTypedElementActorLabelToColumnProcessor::UTypedElementActorLabelToColumnProcessor()
+	: Query(*this)
+{
+	ExecutionFlags = static_cast<int32>(EProcessorExecutionFlags::Editor);
+	ExecutionOrder.ExecuteInGroup = UE::Mass::ProcessorGroupNames::SyncWorldToMass;
+	ProcessingPhase = EMassProcessingPhase::PrePhysics;
+	bRequiresGameThreadExecution = true;
+}
+
+void UTypedElementActorLabelToColumnProcessor::ConfigureQueries()
+{
+	Query.AddRequirement(FMassActorFragment::StaticStruct(), EMassFragmentAccess::ReadOnly);
+	Query.AddRequirement(FTypedElementLabelColumn::StaticStruct(), EMassFragmentAccess::ReadWrite);
+	Query.AddRequirement(FTypedElementLabelHashColumn::StaticStruct(), EMassFragmentAccess::ReadWrite);
+}
+
+void UTypedElementActorLabelToColumnProcessor::Execute(FMassEntityManager& EntityManager, FMassExecutionContext& Context)
+{
+	Query.ForEachEntityChunk(EntityManager, Context, [](FMassExecutionContext& Context)
+		{
+			const FMassActorFragment* ActorIt = Context.GetFragmentView<FMassActorFragment>().GetData();
+			FTypedElementLabelColumn* LabelIt = Context.GetMutableFragmentView<FTypedElementLabelColumn>().GetData();
+			FTypedElementLabelHashColumn* LabelHashIt = Context.GetMutableFragmentView<FTypedElementLabelHashColumn>().GetData();
+			
+			const int32 NumEntities = Context.GetNumEntities();
+			for (int32 Index = 0; Index < NumEntities; ++Index)
+			{
+				const FString& ActorLabel = ActorIt->Get()->GetActorLabel(false);
+				uint64 ActorLabelHash = CityHash64(
+					reinterpret_cast<const char*>(ActorLabel.GetCharArray().GetData()), 
+					ActorLabel.Len() * sizeof(ActorLabel.GetCharArray().GetData()[0]));
+				if (LabelHashIt->LabelHash != ActorLabelHash)
+				{
+					LabelIt->Label = ActorLabel;
+					LabelHashIt->LabelHash = ActorLabelHash;
+				}
+				
+				++ActorIt;
+				++LabelIt;
+				++LabelHashIt;
+			}
+		}
+	);
+}
+
+
+
+/**
+ * UTypedElementLabelColumnToActorProcessor
+ */
+
+UTypedElementLabelColumnToActorProcessor::UTypedElementLabelColumnToActorProcessor()
+	: Query(*this)
+{
+	ExecutionFlags = static_cast<int32>(EProcessorExecutionFlags::Editor);
+	ExecutionOrder.ExecuteInGroup = UE::Mass::ProcessorGroupNames::UpdateWorldFromMass;
+	ProcessingPhase = EMassProcessingPhase::FrameEnd;
+	bRequiresGameThreadExecution = true;
+}
+
+void UTypedElementLabelColumnToActorProcessor::ConfigureQueries()
+{
+	Query.AddRequirement(FMassActorFragment::StaticStruct(), EMassFragmentAccess::ReadWrite);
+	Query.AddRequirement(FTypedElementLabelColumn::StaticStruct(), EMassFragmentAccess::ReadOnly);
+	Query.AddRequirement(FTypedElementLabelHashColumn::StaticStruct(), EMassFragmentAccess::ReadOnly);
+	Query.AddTagRequirement(*FTypedElementSyncBackToWorldTag::StaticStruct(), EMassFragmentPresence::All);
+}
+
+void UTypedElementLabelColumnToActorProcessor::Execute(FMassEntityManager& EntityManager, FMassExecutionContext& Context)
+{
+	Query.ForEachEntityChunk(EntityManager, Context, [](FMassExecutionContext& Context)
+		{
+			FMassActorFragment* ActorIt = Context.GetMutableFragmentView<FMassActorFragment>().GetData();
+			const FTypedElementLabelColumn* LabelIt = Context.GetFragmentView<FTypedElementLabelColumn>().GetData();
+			const FTypedElementLabelHashColumn* LabelHashIt = Context.GetFragmentView<FTypedElementLabelHashColumn>().GetData();
+
+			const int32 NumEntities = Context.GetNumEntities();
+			for (int32 Index = 0; Index < NumEntities; ++Index)
+			{
+				const FString& ActorLabel = ActorIt->Get()->GetActorLabel(false);
+				uint64 ActorLabelHash = CityHash64(
+					reinterpret_cast<const char*>(ActorLabel.GetCharArray().GetData()),
+					ActorLabel.Len() * sizeof(ActorLabel.GetCharArray().GetData()[0]));
+				if (LabelHashIt->LabelHash != ActorLabelHash)
+				{
+					ActorIt->GetMutable()->SetActorLabel(LabelIt->Label);
+				}
+				
+				++ActorIt;
+				++LabelIt;
+				++LabelHashIt;
+			}
+		}
+	);
+}
diff --git a/Engine/Plugins/Experimental/TypedElementsDataStorage/Source/TypedElementsDataStorage/Private/Compatibility/TypedElementActorLabelProcessors.h b/Engine/Plugins/Experimental/TypedElementsDataStorage/Source/TypedElementsDataStorage/Private/Compatibility/TypedElementActorLabelProcessors.h
new file mode 100644
--- /dev/null
+++ b/Engine/Plugins/Experimental/TypedElementsDataStorage/Source/TypedElementsDataStorage/Private/Compatibility/TypedElementActorLabelProcessors.h
@@ -0,0 +1,49 @@
+// Copyright Epic Games, Inc. All Rights Reserved.
+
+#pragma once
+
+#include "MassProcessor.h"
+#include "UObject/ObjectMacros.h"
+
+#include "TypedElementActorLabelProcessors.generated.h"
+
+/**
+ * Takes the label set on an actor and copies it to the Data Storage if they differ.
+ */
+UCLASS()
+class UTypedElementActorLabelToColumnProcessor : public UMassProcessor
+{
+	GENERATED_BODY()
+
+public:
+	UTypedElementActorLabelToColumnProcessor();
+
+protected:
+	void ConfigureQueries() override;
+	void Execute(FMassEntityManager& EntityManager, FMassExecutionContext& Context) override;
+
+private:
+	FMassEntityQuery Query;
+};
+
+
+
+/**
+ * Takes the label stored in the Data Storage and copies it to the actor's label if the FTypedElementSyncBackToWorldTag
+ * has been set and the labels differ.
+ */
+UCLASS()
+class UTypedElementLabelColumnToActorProcessor : public UMassProcessor
+{
+	GENERATED_BODY()
+
+public:
+	UTypedElementLabelColumnToActorProcessor();
+
+protected:
+	void ConfigureQueries() override;
+	void Execute(FMassEntityManager& EntityManager, FMassExecutionContext& Context) override;
+
+private:
+	FMassEntityQuery Query;
+};
\ No newline at end of file
diff --git a/Engine/Plugins/Experimental/TypedElementsDataStorage/Source/TypedElementsDataStorage/Private/Compatibility/TypedElementActorPackagePathToColumnProcessor.cpp b/Engine/Plugins/Experimental/TypedElementsDataStorage/Source/TypedElementsDataStorage/Private/Compatibility/TypedElementActorPackagePathToColumnProcessor.cpp
new file mode 100644
--- /dev/null
+++ b/Engine/Plugins/Experimental/TypedElementsDataStorage/Source/TypedElementsDataStorage/Private/Compatibility/TypedElementActorPackagePathToColumnProcessor.cpp
@@ -0,0 +1,50 @@
+// Copyright Epic Games, Inc. All Rights Reserved.
+
+#include "Compatibility/TypedElementActorPackagePathToColumnProcessor.h"
+
+#include "Elements/Columns/TypedElementPackageColumns.h"
+#include "MassActorSubsystem.h"
+#include "MassExecutionContext.h"
+#include "TypedElementSubsystems.h"
+
+UTypedElementActorPackagePathToColumnProcessor::UTypedElementActorPackagePathToColumnProcessor()
+	: Query(*this)
+{
+	ExecutionFlags = static_cast<int32>(EProcessorExecutionFlags::Editor);
+	ObservedType = FMassActorFragment::StaticStruct();
+	Operation = EMassObservedOperation::Add;
+	bRequiresGameThreadExecution = true;
+}
+
+void UTypedElementActorPackagePathToColumnProcessor::ConfigureQueries()
+{
+	Query.AddRequirement(FMassActorFragment::StaticStruct(), EMassFragmentAccess::ReadOnly);
+	Query.AddRequirement(FTypedElementPackagePathColumn::StaticStruct(), EMassFragmentAccess::ReadWrite);
+	Query.AddRequirement(FTypedElementPackageLoadedPathColumn::StaticStruct(), EMassFragmentAccess::ReadWrite);
+	
+	ProcessorRequirements.AddSubsystemRequirement(UTypedElementDataStorageSubsystem::StaticClass(), EMassFragmentAccess::ReadWrite);
+}
+
+void UTypedElementActorPackagePathToColumnProcessor::Execute(FMassEntityManager& EntityManager, FMassExecutionContext& Context)
+{
+	Query.ForEachEntityChunk(EntityManager, Context, [](FMassExecutionContext& Context)
+		{
+			const FMassActorFragment* ActorIt = Context.GetFragmentView<FMassActorFragment>().GetData();
+			FTypedElementPackagePathColumn* PathIt = Context.GetMutableFragmentView<FTypedElementPackagePathColumn>().GetData();
+			FTypedElementPackageLoadedPathColumn* LoadedPathIt = Context.GetMutableFragmentView<FTypedElementPackageLoadedPathColumn>().GetData();
+			
+			const int32 EntityCount = Context.GetNumEntities();
+			for (int32 Counter=0; Counter < EntityCount; ++Counter)
+			{
+				UPackage* Target = ActorIt->Get()->GetPackage();
+
+				Target->GetPathName(nullptr, PathIt->Path);
+				LoadedPathIt->LoadedPath = Target->GetLoadedPath();
+				
+				++ActorIt;
+				++PathIt;
+				++LoadedPathIt;
+			}
+		}
+	);
+}
diff --git a/Engine/Plugins/Experimental/TypedElementsDataStorage/Source/TypedElementsDataStorage/Private/Compatibility/TypedElementActorPackagePathToColumnProcessor.h b/Engine/Plugins/Experimental/TypedElementsDataStorage/Source/TypedElementsDataStorage/Private/Compatibility/TypedElementActorPackagePathToColumnProcessor.h
new file mode 100644
--- /dev/null
+++ b/Engine/Plugins/Experimental/TypedElementsDataStorage/Source/TypedElementsDataStorage/Private/Compatibility/TypedElementActorPackagePathToColumnProcessor.h
@@ -0,0 +1,24 @@
+// Copyright Epic Games, Inc. All Rights Reserved.
+
+#pragma once
+
+#include "MassObserverProcessor.h"
+#include "UObject/ObjectMacros.h"
+
+#include "TypedElementActorPackagePathToColumnProcessor.generated.h"
+
+UCLASS()
+class UTypedElementActorPackagePathToColumnProcessor : public UMassObserverProcessor
+{
+	GENERATED_BODY()
+
+public:
+	UTypedElementActorPackagePathToColumnProcessor();
+
+protected:
+	void ConfigureQueries() override;
+	void Execute(FMassEntityManager& EntityManager, FMassExecutionContext& Context) override;
+
+private:
+	FMassEntityQuery Query;
+};
diff --git a/Engine/Plugins/Experimental/TypedElementsDataStorage/Source/TypedElementsDataStorage/Private/Compatibility/TypedElementActorTransformProcessors.cpp b/Engine/Plugins/Experimental/TypedElementsDataStorage/Source/TypedElementsDataStorage/Private/Compatibility/TypedElementActorTransformProcessors.cpp
new file mode 100644
--- /dev/null
+++ b/Engine/Plugins/Experimental/TypedElementsDataStorage/Source/TypedElementsDataStorage/Private/Compatibility/TypedElementActorTransformProcessors.cpp
@@ -0,0 +1,133 @@
+// Copyright Epic Games, Inc. All Rights Reserved.
+
+#include "Compatibility/TypedElementActorTransformProcessors.h"
+
+#include "Elements/Columns/TypedElementMiscColumns.h"
+#include "Elements/Columns/TypedElementTransformColumns.h"
+#include "MassActorSubsystem.h"
+#include "MassExecutionContext.h"
+
+/**
+ * UTypedElementActorAddTransformColumnProcessor
+ */
+
+UTypedElementActorAddTransformColumnProcessor::UTypedElementActorAddTransformColumnProcessor()
+	: Query(*this)
+{
+	ExecutionFlags = static_cast<int32>(EProcessorExecutionFlags::Editor);
+	// Run this before the world gets synced to MASS so a new transform is initialized in the same frame.
+	ExecutionOrder.ExecuteBefore.Add(UE::Mass::ProcessorGroupNames::SyncWorldToMass);
+	ProcessingPhase = EMassProcessingPhase::PrePhysics;
+	bRequiresGameThreadExecution = true;
+}
+
+void UTypedElementActorAddTransformColumnProcessor::ConfigureQueries()
+{
+	Query.AddRequirement(FMassActorFragment::StaticStruct(), EMassFragmentAccess::ReadOnly);
+	Query.AddRequirement(FTypedElementLocalTransformColumn::StaticStruct(), EMassFragmentAccess::ReadOnly, EMassFragmentPresence::None);
+}
+
+void UTypedElementActorAddTransformColumnProcessor::Execute(FMassEntityManager& EntityManager, FMassExecutionContext& Context)
+{
+	Query.ForEachEntityChunk(EntityManager, Context, [](FMassExecutionContext& Context)
+		{
+			const FMassActorFragment* ActorIt = Context.GetFragmentView<FMassActorFragment>().GetData();
+			
+			for (FMassEntityHandle Entity : Context.GetEntities())
+			{
+				if (ActorIt->Get()->GetRootComponent())
+				{
+					Context.Defer().AddFragment_RuntimeCheck<FTypedElementLocalTransformColumn>(Entity);
+				}
+				
+				++ActorIt;
+			}
+		}
+	);
+}
+
+
+/**
+ * UTypedElementActorLocalTransformToColumnProcessor
+ */
+
+UTypedElementActorLocalTransformToColumnProcessor::UTypedElementActorLocalTransformToColumnProcessor()
+	: Query(*this)
+{
+	ExecutionFlags = static_cast<int32>(EProcessorExecutionFlags::Editor);
+	ExecutionOrder.ExecuteInGroup = UE::Mass::ProcessorGroupNames::SyncWorldToMass;
+	ProcessingPhase = EMassProcessingPhase::PrePhysics;
+	bRequiresGameThreadExecution = true;
+}
+
+void UTypedElementActorLocalTransformToColumnProcessor::ConfigureQueries()
+{
+	Query.AddRequirement(FMassActorFragment::StaticStruct(), EMassFragmentAccess::ReadOnly);
+	Query.AddRequirement(FTypedElementLocalTransformColumn::StaticStruct(), EMassFragmentAccess::ReadWrite);
+}
+
+void UTypedElementActorLocalTransformToColumnProcessor::Execute(FMassEntityManager& EntityManager, FMassExecutionContext& Context)
+{
+	Query.ForEachEntityChunk(EntityManager, Context, [](FMassExecutionContext& Context)
+		{
+			const FMassActorFragment* ActorIt = Context.GetFragmentView<FMassActorFragment>().GetData();
+			FTypedElementLocalTransformColumn* TransformIt = Context.GetMutableFragmentView<FTypedElementLocalTransformColumn>().GetData();
+
+			for (FMassEntityHandle Entity : Context.GetEntities())
+			{
+				const AActor* ActorInstance = ActorIt->Get();
+				if (ActorInstance != nullptr && ActorInstance->GetRootComponent() != nullptr)
+				{
+					TransformIt->Transform = ActorIt->Get()->GetActorTransform();
+				}
+				else
+				{
+					Context.Defer().RemoveFragment_RuntimeCheck<FTypedElementLocalTransformColumn>(Entity);
+				}
+				
+				++ActorIt;
+				++TransformIt;
+			}
+		}
+	);
+}
+
+
+/**
+ * UTypedElementTransformColumnToActorProcessor
+ */
+
+UTypedElementTransformColumnToActorProcessor::UTypedElementTransformColumnToActorProcessor()
+	: Query(*this)
+{
+	ExecutionFlags = static_cast<int32>(EProcessorExecutionFlags::Editor);
+	ExecutionOrder.ExecuteInGroup = UE::Mass::ProcessorGroupNames::UpdateWorldFromMass;
+	ProcessingPhase = EMassProcessingPhase::FrameEnd;
+	bRequiresGameThreadExecution = true;
+}
+
+void UTypedElementTransformColumnToActorProcessor::ConfigureQueries()
+{
+	Query.AddRequirement(FMassActorFragment::StaticStruct(), EMassFragmentAccess::ReadWrite);
+	Query.AddRequirement(FTypedElementLocalTransformColumn::StaticStruct(), EMassFragmentAccess::ReadOnly);
+	Query.AddTagRequirement(*FTypedElementSyncBackToWorldTag::StaticStruct(), EMassFragmentPresence::All);
+}
+
+void UTypedElementTransformColumnToActorProcessor::Execute(FMassEntityManager& EntityManager, FMassExecutionContext& Context)
+{
+	Query.ForEachEntityChunk(EntityManager, Context, [](FMassExecutionContext& Context)
+		{
+			FMassActorFragment* ActorIt = Context.GetMutableFragmentView<FMassActorFragment>().GetData();
+			const FTypedElementLocalTransformColumn* TransformIt = Context.GetFragmentView<FTypedElementLocalTransformColumn>().GetData();
+			
+			const int32 NumEntities = Context.GetNumEntities();
+			for (int32 Index = 0; Index < NumEntities; ++Index)
+			{
+				ActorIt->GetMutable()->SetActorTransform(TransformIt->Transform);
+				
+				++ActorIt;
+				++TransformIt;
+			}
+		}
+	);
+}
diff --git a/Engine/Plugins/Experimental/TypedElementsDataStorage/Source/TypedElementsDataStorage/Private/Compatibility/TypedElementActorTransformProcessors.h b/Engine/Plugins/Experimental/TypedElementsDataStorage/Source/TypedElementsDataStorage/Private/Compatibility/TypedElementActorTransformProcessors.h
new file mode 100644
--- /dev/null
+++ b/Engine/Plugins/Experimental/TypedElementsDataStorage/Source/TypedElementsDataStorage/Private/Compatibility/TypedElementActorTransformProcessors.h
@@ -0,0 +1,71 @@
+// Copyright Epic Games, Inc. All Rights Reserved.
+
+#pragma once
+
+#include "MassProcessor.h"
+#include "UObject/ObjectMacros.h"
+
+#include "TypedElementActorTransformProcessors.generated.h"
+
+/**
+ * Checks actors that don't have a tranform column and adds one if an actor has been
+ * assigned a transform.
+ */
+UCLASS()
+class UTypedElementActorAddTransformColumnProcessor : public UMassProcessor
+{
+	GENERATED_BODY()
+
+public:
+	UTypedElementActorAddTransformColumnProcessor();
+
+protected:
+	void ConfigureQueries() override;
+	void Execute(FMassEntityManager& EntityManager, FMassExecutionContext& Context) override;
+
+private:
+	FMassEntityQuery Query;
+};
+
+
+/**
+ * Takes the transform set on an actor and copies it to the Data Storage or removes the 
+ * transform column if there's not transform available anymore.
+ */
+UCLASS()
+class UTypedElementActorLocalTransformToColumnProcessor : public UMassProcessor
+{
+	GENERATED_BODY()
+
+public:
+	UTypedElementActorLocalTransformToColumnProcessor();
+
+protected:
+	void ConfigureQueries() override;
+	void Execute(FMassEntityManager& EntityManager, FMassExecutionContext& Context) override;
+
+private:
+	FMassEntityQuery Query;
+};
+
+
+
+/**
+ * Takes the transform stored in the Data Storage and copies it to the actor's tranform if 
+ * the FTypedElementSyncBackToWorldTag has been set.
+ */
+UCLASS()
+class UTypedElementTransformColumnToActorProcessor : public UMassProcessor
+{
+	GENERATED_BODY()
+
+public:
+	UTypedElementTransformColumnToActorProcessor();
+
+protected:
+	void ConfigureQueries() override;
+	void Execute(FMassEntityManager& EntityManager, FMassExecutionContext& Context) override;
+
+private:
+	FMassEntityQuery Query;
+};
\ No newline at end of file
diff --git a/Engine/Plugins/Experimental/TypedElementsDataStorage/Source/TypedElementsDataStorage/Private/Processors/TypedElementMiscProcessors.cpp b/Engine/Plugins/Experimental/TypedElementsDataStorage/Source/TypedElementsDataStorage/Private/Processors/TypedElementMiscProcessors.cpp
new file mode 100644
--- /dev/null
+++ b/Engine/Plugins/Experimental/TypedElementsDataStorage/Source/TypedElementsDataStorage/Private/Processors/TypedElementMiscProcessors.cpp
@@ -0,0 +1,36 @@
+// Copyright Epic Games, Inc. All Rights Reserved.
+
+#include "Processors/TypedElementMiscProcessors.h"
+
+#include "MassCommonTypes.h"
+#include "MassExecutionContext.h"
+#include "Elements/Columns/TypedElementMiscColumns.h"
+
+/**
+ * UTypedElementRemoveSyncToWorldTagProcessor
+ */
+
+UTypedElementRemoveSyncToWorldTagProcessor::UTypedElementRemoveSyncToWorldTagProcessor()
+	: Query(*this)
+{
+	ExecutionFlags = static_cast<int32>(EProcessorExecutionFlags::Editor);
+	ExecutionOrder.ExecuteAfter.Add(UE::Mass::ProcessorGroupNames::UpdateWorldFromMass);
+	ProcessingPhase = EMassProcessingPhase::FrameEnd;
+}
+
+void UTypedElementRemoveSyncToWorldTagProcessor::ConfigureQueries()
+{
+	Query.AddTagRequirement(*FTypedElementSyncBackToWorldTag::StaticStruct(), EMassFragmentPresence::All);
+}
+
+void UTypedElementRemoveSyncToWorldTagProcessor::Execute(FMassEntityManager& EntityManager, FMassExecutionContext& Context)
+{
+	Query.ForEachEntityChunk(EntityManager, Context, [](FMassExecutionContext& Context)
+		{
+			for (FMassEntityHandle Entity : Context.GetEntities())
+			{
+				Context.Defer().RemoveTag_RuntimeCheck<FTypedElementSyncBackToWorldTag>(Entity);
+			}
+		}
+	);
+}
diff --git a/Engine/Plugins/Experimental/TypedElementsDataStorage/Source/TypedElementsDataStorage/Private/Processors/TypedElementMiscProcessors.h b/Engine/Plugins/Experimental/TypedElementsDataStorage/Source/TypedElementsDataStorage/Private/Processors/TypedElementMiscProcessors.h
new file mode 100644
--- /dev/null
+++ b/Engine/Plugins/Experimental/TypedElementsDataStorage/Source/TypedElementsDataStorage/Private/Processors/TypedElementMiscProcessors.h
@@ -0,0 +1,27 @@
+// Copyright Epic Games, Inc. All Rights Reserved.
+
+#pragma once
+
+#include "MassProcessor.h"
+#include "UObject/ObjectMacros.h"
+
+#include "TypedElementMiscProcessors.generated.h"
+
+/**
+ * Removes all FTypedElementSyncBackToWorldTags at the end of an update cycle.
+ */
+UCLASS()
+class UTypedElementRemoveSyncToWorldTagProcessor : public UMassProcessor
+{
+	GENERATED_BODY()
+
+public:
+	UTypedElementRemoveSyncToWorldTagProcessor();
+
+protected:
+	void ConfigureQueries() override;
+	void Execute(FMassEntityManager& EntityManager, FMassExecutionContext& Context) override;
+
+private:
+	FMassEntityQuery Query;
+};
\ No newline at end of file
diff --git a/Engine/Plugins/Experimental/TypedElementsDataStorage/Source/TypedElementsDataStorage/Private/TypedElementDatabase.cpp b/Engine/Plugins/Experimental/TypedElementsDataStorage/Source/TypedElementsDataStorage/Private/TypedElementDatabase.cpp
new file mode 100644
--- /dev/null
+++ b/Engine/Plugins/Experimental/TypedElementsDataStorage/Source/TypedElementsDataStorage/Private/TypedElementDatabase.cpp
@@ -0,0 +1,423 @@
+// Copyright Epic Games, Inc. All Rights Reserved.
+
+#include "TypedElementDatabase.h"
+
+#include "Editor.h"
+#include "Engine/World.h"
+#include "MassEntityEditorSubsystem.h"
+#include "Stats/Stats2.h"
+#include "TickTaskManagerInterface.h"
+#include "MassSubsystemAccess.h"
+
+void UTypedElementDatabase::Initialize()
+{
+	check(GEditor);
+	UMassEntityEditorSubsystem* Mass = GEditor->GetEditorSubsystem<UMassEntityEditorSubsystem>();
+	check(Mass);
+	Mass->GetOnPreTickDelegate().AddUObject(this, &UTypedElementDatabase::OnPreMassTick);
+
+	ActiveEditorEntityManager = Mass->GetMutableEntityManager().AsShared();
+}
+
+void UTypedElementDatabase::Deinitialize()
+{
+	Reset();
+}
+
+void UTypedElementDatabase::OnPreMassTick(float DeltaTime)
+{
+	checkf(IsAvailable(), TEXT("Typed Element Database was ticked while it's not ready."));
+	OnUpdateDelegate.Broadcast();
+}
+
+TSharedPtr<FMassEntityManager> UTypedElementDatabase::GetActiveMutableEditorEntityManager()
+{
+	return ActiveEditorEntityManager;
+}
+
+TSharedPtr<const FMassEntityManager> UTypedElementDatabase::GetActiveEditorEntityManager() const
+{
+	return ActiveEditorEntityManager;
+}
+
+TypedElementTableHandle UTypedElementDatabase::RegisterTable(TConstArrayView<const UScriptStruct*> ColumnList)
+{
+	return RegisterTable(ColumnList, {});
+}
+
+TypedElementTableHandle UTypedElementDatabase::RegisterTable(TConstArrayView<const UScriptStruct*> ColumnList, const FName Name)
+{
+	if (ActiveEditorEntityManager && (!Name.IsValid() || !TableNameLookup.Contains(Name)))
+	{
+		TypedElementTableHandle Result = Tables.Num();
+		Tables.Add(ActiveEditorEntityManager->CreateArchetype(ColumnList, Name));
+		if (Name.IsValid())
+		{
+			TableNameLookup.Add(Name, Result);
+		}
+		return Result;
+	}
+	return TypedElementInvalidTableHandle;
+}
+
+TypedElementTableHandle UTypedElementDatabase::RegisterTable(TypedElementTableHandle SourceTable,
+	TConstArrayView<const UScriptStruct*> ColumnList)
+{
+	return RegisterTable(SourceTable, ColumnList, {});
+}
+
+TypedElementTableHandle UTypedElementDatabase::RegisterTable(TypedElementTableHandle SourceTable, 
+	TConstArrayView<const UScriptStruct*> ColumnList, const FName Name)
+{
+	if (ActiveEditorEntityManager && (!Name.IsValid() || !TableNameLookup.Contains(Name)) && SourceTable < Tables.Num())
+	{
+		TypedElementTableHandle Result = Tables.Num();
+		Tables.Add(ActiveEditorEntityManager->CreateArchetype(Tables[SourceTable], ColumnList, Name));
+		if (Name.IsValid())
+		{
+			TableNameLookup.Add(Name, Result);
+		}
+		return Result;
+	}
+	return TypedElementInvalidTableHandle;
+}
+
+TypedElementTableHandle UTypedElementDatabase::FindTable(const FName Name)
+{
+	TypedElementTableHandle* TableHandle = TableNameLookup.Find(Name);
+	return TableHandle ? *TableHandle : TypedElementInvalidTableHandle;
+}
+
+TypedElementRowHandle UTypedElementDatabase::AddRow(TypedElementTableHandle Table)
+{
+	checkf(Table < Tables.Num(), TEXT("Attempting to add a row to a non-existing table."));
+	return ActiveEditorEntityManager ? 
+		ActiveEditorEntityManager->CreateEntity(Tables[Table]).AsNumber() :
+		TypedElementInvalidRowHandle;
+}
+
+TypedElementRowHandle UTypedElementDatabase::AddRow(FName TableName)
+{
+	TypedElementTableHandle* Table = TableNameLookup.Find(TableName);
+	return Table ? AddRow(*Table) : TypedElementInvalidRowHandle;
+}
+
+bool UTypedElementDatabase::BatchAddRow(TypedElementTableHandle Table, int32 Count, TypedElementDataStorageCreationCallbackRef OnCreated)
+{
+	OnCreated.CheckCallable();
+	checkf(Table < Tables.Num(), TEXT("Attempting to add multiple rows to a non-existing table."));
+	if (ActiveEditorEntityManager)
+	{
+		TArray<FMassEntityHandle> Entities;
+		Entities.Reserve(Count);
+		TSharedRef<FMassEntityManager::FEntityCreationContext> Context = 
+			ActiveEditorEntityManager->BatchCreateEntities(Tables[Table], Count, Entities);
+		
+		for (FMassEntityHandle Entity : Entities)
+		{
+			OnCreated(Entity.AsNumber());
+		}
+
+		return true;
+	}
+	return false;
+}
+
+bool UTypedElementDatabase::BatchAddRow(FName TableName, int32 Count, TypedElementDataStorageCreationCallbackRef OnCreated)
+{
+	TypedElementTableHandle* Table = TableNameLookup.Find(TableName);
+	return Table ? BatchAddRow(*Table, Count, OnCreated) : false;
+}
+
+void UTypedElementDatabase::RemoveRow(TypedElementRowHandle Row)
+{
+	if (ActiveEditorEntityManager)
+	{
+		ActiveEditorEntityManager->DestroyEntity(FMassEntityHandle::FromNumber(Row));
+	}
+}
+
+void UTypedElementDatabase::AddTag(TypedElementRowHandle Row, const UScriptStruct* TagType)
+{
+	checkf(TagType && TagType->IsChildOf(FMassTag::StaticStruct()),
+		TEXT("Tag type '%s' is invalid as it needs to be set or derived from FMassTag."), *GetPathNameSafe(TagType));
+
+	FMassEntityHandle Entity = FMassEntityHandle::FromNumber(Row);
+	if (ActiveEditorEntityManager && ActiveEditorEntityManager->IsEntityValid(Entity))
+	{
+		ActiveEditorEntityManager->AddTagToEntity(Entity, TagType);
+	}
+}
+
+void UTypedElementDatabase::AddTag(TypedElementRowHandle Row, FTopLevelAssetPath TagName)
+{
+	bool bExactMatch = true;
+	UScriptStruct* TagStructInfo = Cast<UScriptStruct>(StaticFindObject(UScriptStruct::StaticClass(), TagName, bExactMatch));
+	if (TagStructInfo)
+	{
+		AddTag(Row, TagStructInfo);
+	}
+}
+
+void* UTypedElementDatabase::AddOrGetColumnData(TypedElementRowHandle Row, const UScriptStruct* ColumnType)
+{
+	checkf(ColumnType && ColumnType->IsChildOf(FMassFragment::StaticStruct()), 
+		TEXT("Colum type '%s' is invalid as it needs to be set or derived from FMassFragment."), *GetPathNameSafe(ColumnType));
+
+	FMassEntityHandle Entity = FMassEntityHandle::FromNumber(Row);
+	if (ActiveEditorEntityManager && ActiveEditorEntityManager->IsEntityValid(Entity))
+	{
+		FStructView Column = ActiveEditorEntityManager->GetFragmentDataStruct(Entity, ColumnType);
+		if (!Column.IsValid())
+		{
+			ActiveEditorEntityManager->AddFragmentToEntity(Entity, ColumnType);
+			Column = ActiveEditorEntityManager->GetFragmentDataStruct(Entity, ColumnType);
+			checkf(Column.IsValid(), TEXT("Added a new column to the Typed Element's data storae, but it couldn't be retrieved."));
+
+		}
+		return Column.GetMemory();
+	}
+	return nullptr;
+}
+
+ColumnDataResult UTypedElementDatabase::AddOrGetColumnData(TypedElementRowHandle Row, FTopLevelAssetPath ColumnName)
+{
+	constexpr bool bExactMatch = true;
+	UScriptStruct* FragmentStructInfo = Cast<UScriptStruct>(StaticFindObject(UScriptStruct::StaticClass(), ColumnName, bExactMatch));
+	return FragmentStructInfo ?
+		ColumnDataResult{ FragmentStructInfo, AddOrGetColumnData(Row, FragmentStructInfo) }:
+		ColumnDataResult{ nullptr, nullptr };
+}
+
+ColumnDataResult UTypedElementDatabase::AddOrGetColumnData(TypedElementRowHandle Row, FTopLevelAssetPath ColumnName,
+	TConstArrayView<TypedElement::ColumnUtils::Argument> Arguments)
+{
+	ColumnDataResult Result = AddOrGetColumnData(Row, ColumnName);
+	if (Result.Description && Result.Data)
+	{
+		TypedElement::ColumnUtils::SetColumnValues(Result.Data, Result.Description, Arguments);
+		return Result;
+	}
+	else
+	{
+		return ColumnDataResult{ nullptr, nullptr };
+	}
+}
+
+void* UTypedElementDatabase::GetColumnData(TypedElementRowHandle Row, const UScriptStruct* ColumnType)
+{
+	checkf(ColumnType && ColumnType->IsChildOf(FMassFragment::StaticStruct()),
+		TEXT("Colum type '%s' is invalid as it needs to be set or derived from FMassFragment."), *GetPathNameSafe(ColumnType));
+
+	FMassEntityHandle Entity = FMassEntityHandle::FromNumber(Row);
+	if (ActiveEditorEntityManager && ActiveEditorEntityManager->IsEntityValid(Entity))
+	{
+		FStructView Column = ActiveEditorEntityManager->GetFragmentDataStruct(Entity, ColumnType);
+		if (Column.IsValid())
+		{
+			return Column.GetMemory();
+		}
+	}
+	return nullptr;
+}
+
+ColumnDataResult UTypedElementDatabase::GetColumnData(TypedElementRowHandle Row, FTopLevelAssetPath ColumnName)
+{
+	FMassEntityHandle Entity = FMassEntityHandle::FromNumber(Row);
+	if (ActiveEditorEntityManager && ActiveEditorEntityManager->IsEntityValid(Entity))
+	{
+		const UScriptStruct* FragmentType = nullptr;
+		FMassArchetypeHandle Archetype = ActiveEditorEntityManager->GetArchetypeForEntityUnsafe(Entity);
+		ActiveEditorEntityManager->ForEachArchetypeFragmentType(Archetype, 
+			[ColumnName, &FragmentType](const UScriptStruct* Fragment)
+			{
+				if (Fragment->GetStructPathName() == ColumnName)
+				{
+					FragmentType = Fragment;
+				}
+			});
+
+		if (FragmentType != nullptr)
+		{
+			FStructView Column = ActiveEditorEntityManager->GetFragmentDataStruct(Entity, FragmentType);
+			if (Column.IsValid())
+			{
+				return ColumnDataResult{ FragmentType, Column.GetMemory() };
+			}
+		}
+	}
+	return ColumnDataResult{ nullptr, nullptr };
+}
+
+TypedElementQueryHandle UTypedElementDatabase::RegisterQuery(const FQueryDescription& Query)
+{
+	auto LocalToNativeAccess = [](FQueryDescription::EAccessType Access) -> EMassFragmentAccess
+	{
+		switch (Access)
+		{
+		case FQueryDescription::EAccessType::ReadOnly:
+			return EMassFragmentAccess::ReadOnly;
+		case FQueryDescription::EAccessType::ReadWrite:
+			return EMassFragmentAccess::ReadWrite;
+		default:
+			checkf(false, TEXT("Invalid query access type: %i."), static_cast<uint32>(Access));
+			return EMassFragmentAccess::MAX;
+		}
+	};
+
+	QueryStore::Handle Result = Queries.Emplace();
+	FExtendedQuery& StoredQuery = Queries.GetMutable(Result);
+	StoredQuery.Action = Query.Action;
+	StoredQuery.bSimpleQuery = Query.bSimpleQuery;
+
+	FMassEntityQuery& NativeQuery = StoredQuery.NativeQuery;
+
+	if (Query.Action == FQueryDescription::EActionType::Count)
+	{
+		checkf(Query.Selection.IsEmpty(), TEXT("Count queries for the Typed Elements Data Storage can't have entries for selection."));
+	}
+	else
+	{
+		for (const FQueryDescription::FAccessControlledStruct& SelectEntry : Query.Selection)
+		{
+			checkf(SelectEntry.Type, TEXT("Provided query selection type can not be null."));
+			checkf(SelectEntry.Type->IsChildOf(FTypedElementDataStorageColumn::StaticStruct()),
+				TEXT("Provided query selection type '%s' is not based on FTypedElementDataStorageColumn."), *SelectEntry.Type->GetStructPathName().ToString());
+			NativeQuery.AddRequirement(SelectEntry.Type, LocalToNativeAccess(SelectEntry.Access));
+		}
+	}
+
+	if (Query.bSimpleQuery) // This backend currently only supports simple queries.
+	{
+		checkf(Query.ConditionTypes.Num() == Query.ConditionOperators.Num(),
+			TEXT("The types and operators for a typed element query have gone out of sync."));
+		
+		const FQueryDescription::FOperator* Operand = Query.ConditionOperators.GetData();
+		for (FQueryDescription::EOperatorType Type : Query.ConditionTypes)
+		{
+			EMassFragmentPresence Presence;
+			switch (Type)
+			{
+			case FQueryDescription::EOperatorType::SimpleAll:
+				Presence = EMassFragmentPresence::All;
+				break;
+			case FQueryDescription::EOperatorType::SimpleAny:
+				Presence = EMassFragmentPresence::Any;
+				break;
+			case FQueryDescription::EOperatorType::SimpleNone:
+				Presence = EMassFragmentPresence::None;
+				break;
+			default:
+				continue;
+			}
+
+			if (Operand->Type->IsChildOf(FMassTag::StaticStruct()))
+			{
+				NativeQuery.AddTagRequirement(*(Operand->Type), Presence);
+			}
+			else if (Operand->Type->IsChildOf(FMassFragment::StaticStruct()))
+			{
+				NativeQuery.AddRequirement(Operand->Type, EMassFragmentAccess::None, Presence);
+			}
+
+			++Operand;
+		}
+	}
+
+	for (const FQueryDescription::FAccessControlledClass& DependencyEntry : Query.Dependencies)
+	{
+		checkf(DependencyEntry.Type, TEXT("Provided query dependcy type can not be null."));
+		checkf(DependencyEntry.Type->IsChildOf<UWorldSubsystem>(),
+			TEXT("Provided query dependency type '%s' is not based on UWorldSubSystem."), *DependencyEntry.Type->GetStructPathName().ToString());
+		
+		constexpr bool bGameThreadOnly = true;
+		NativeQuery.AddSubsystemRequirement(const_cast<UClass*>(DependencyEntry.Type), LocalToNativeAccess(DependencyEntry.Access), bGameThreadOnly);
+	}
+
+	return Result.Handle;
+}
+
+void UTypedElementDatabase::UnregisterQuery(TypedElementQueryHandle Query)
+{
+	QueryStore::Handle Handle;
+	Handle.Handle = Query;
+
+	if (Queries.IsAlive(Handle))
+	{
+		Queries.Get(Handle).NativeQuery.Clear();
+	}
+
+	Queries.Remove(Handle);
+}
+
+ITypedElementDataStorageInterface::FQueryResult UTypedElementDatabase::RunQuery(TypedElementQueryHandle Query)
+{
+	FQueryResult Result;
+
+	QueryStore::Handle Handle;
+	Handle.Handle = Query;
+
+	if (Queries.IsAlive(Handle))
+	{
+		FExtendedQuery& QueryData = Queries.Get(Handle);
+		if (QueryData.bSimpleQuery)
+		{
+			switch (QueryData.Action)
+			{
+			case FQueryDescription::EActionType::None:
+				Result.Completed = FQueryResult::ECompletion::Fully;
+				break;
+			case FQueryDescription::EActionType::Select:
+				checkf(false, TEXT("Support for this option will be coming in a future update."));
+				Result.Completed = FQueryResult::ECompletion::Unsupported;
+				break;
+			case FQueryDescription::EActionType::Count:
+				checkf(ActiveEditorEntityManager, 
+					TEXT("Unable to run Typed Element Data Storage query before a MASS Entity Manager has been assigned."));
+				Result.Count = QueryData.NativeQuery.GetNumMatchingEntities(*ActiveEditorEntityManager);
+				Result.Completed = FQueryResult::ECompletion::Fully;
+				break;
+			default:
+				Result.Completed = FQueryResult::ECompletion::Unsupported;
+				break;
+			}
+		}
+		else
+		{
+			checkf(false, TEXT("Support for this option will be coming in a future update."));
+			Result.Completed = FQueryResult::ECompletion::Unsupported;
+		}
+	}
+	else
+	{
+		Result.Completed = FQueryResult::ECompletion::Unavailable;
+	}
+
+	return Result;
+}
+
+FTypedElementOnDataStorageUpdate& UTypedElementDatabase::OnUpdate()
+{
+	return OnUpdateDelegate;
+}
+
+bool UTypedElementDatabase::IsAvailable() const
+{
+	return bool(ActiveEditorEntityManager);
+}
+
+void* UTypedElementDatabase::GetExternalSystemAddress(UClass* Target)
+{
+	if (Target && Target->IsChildOf<USubsystem>())
+	{
+		return FMassSubsystemAccess::FetchSubsystemInstance(/*World=*/nullptr, Target);
+	}
+	return nullptr;
+}
+
+void UTypedElementDatabase::Reset()
+{
+	Tables.Reset();
+	TableNameLookup.Reset();
+	ActiveEditorEntityManager.Reset();
+}
\ No newline at end of file
diff --git a/Engine/Plugins/Experimental/TypedElementsDataStorage/Source/TypedElementsDataStorage/Private/TypedElementDatabase.h b/Engine/Plugins/Experimental/TypedElementsDataStorage/Source/TypedElementsDataStorage/Private/TypedElementDatabase.h
new file mode 100644
--- /dev/null
+++ b/Engine/Plugins/Experimental/TypedElementsDataStorage/Source/TypedElementsDataStorage/Private/TypedElementDatabase.h
@@ -0,0 +1,84 @@
+// Copyright Epic Games, Inc. All Rights Reserved.
+
+#pragma once
+
+#include "Templates/SharedPointer.h"
+#include "Elements/Interfaces/TypedElementDataStorageInterface.h"
+#include "MassArchetypeTypes.h"
+#include "MassEntityQuery.h"
+#include "TypedElementHandleStore.h"
+#include "UObject/ObjectMacros.h"
+
+#include "TypedElementDatabase.generated.h"
+
+struct FMassEntityManager;
+class UWorld;
+
+UCLASS()
+class TYPEDELEMENTSDATASTORAGE_API UTypedElementDatabase 
+	: public UObject
+	, public ITypedElementDataStorageInterface
+{
+	GENERATED_BODY()
+public:
+	~UTypedElementDatabase() override = default;
+	
+	void Initialize();
+	void Deinitialize();
+
+	/** Triggered just before underlying Mass processing gets ticked */
+	void OnPreMassTick(float DeltaTime);
+
+	TSharedPtr<FMassEntityManager> GetActiveMutableEditorEntityManager();
+	TSharedPtr<const FMassEntityManager> GetActiveEditorEntityManager() const;
+
+	TypedElementTableHandle RegisterTable(TConstArrayView<const UScriptStruct*> ColumnList) override;
+	TypedElementTableHandle RegisterTable(TConstArrayView<const UScriptStruct*> ColumnList, const FName Name) override;
+	TypedElementTableHandle RegisterTable(TypedElementTableHandle SourceTable, TConstArrayView<const UScriptStruct*> ColumnList) override;
+	TypedElementTableHandle RegisterTable(TypedElementTableHandle SourceTable, TConstArrayView<const UScriptStruct*> ColumnList, 
+		const FName Name) override;
+	TypedElementTableHandle FindTable(const FName Name) override;
+
+	TypedElementRowHandle AddRow(TypedElementTableHandle Table) override;
+	TypedElementRowHandle AddRow(FName TableName) override;
+	bool BatchAddRow(TypedElementTableHandle Table, int32 Count, TypedElementDataStorageCreationCallbackRef OnCreated) override;
+	bool BatchAddRow(FName TableName, int32 Count, TypedElementDataStorageCreationCallbackRef OnCreated) override;
+	void RemoveRow(TypedElementRowHandle Row) override;
+
+	void AddTag(TypedElementRowHandle Row, const UScriptStruct* TagType) override;
+	void AddTag(TypedElementRowHandle Row, FTopLevelAssetPath TagName) override;
+	void* AddOrGetColumnData(TypedElementRowHandle Row, const UScriptStruct* ColumnType) override;
+	ColumnDataResult AddOrGetColumnData(TypedElementRowHandle Row, FTopLevelAssetPath ColumnName) override;
+	void* GetColumnData(TypedElementRowHandle Row, const UScriptStruct* ColumnType) override;
+	ColumnDataResult AddOrGetColumnData(TypedElementRowHandle Row, FTopLevelAssetPath ColumnName,
+		TConstArrayView<TypedElement::ColumnUtils::Argument> Arguments) override;
+	ColumnDataResult GetColumnData(TypedElementRowHandle Row, FTopLevelAssetPath ColumnName) override;
+
+	TypedElementQueryHandle RegisterQuery(const FQueryDescription& Query) override;
+	void UnregisterQuery(TypedElementQueryHandle Query) override;
+	FQueryResult RunQuery(TypedElementQueryHandle Query) override;
+
+	FTypedElementOnDataStorageUpdate& OnUpdate() override;
+	bool IsAvailable() const override;
+	void* GetExternalSystemAddress(UClass* Target) override;
+
+private:
+	struct FExtendedQuery
+	{
+		FMassEntityQuery NativeQuery;
+		FQueryDescription::EActionType Action{ FQueryDescription::EActionType::None };
+		bool bSimpleQuery{ false };
+	};
+	using QueryStore = TTypedElementHandleStore<FExtendedQuery>;
+
+	void Reset();
+	
+	TArray<FMassArchetypeHandle> Tables;
+	TMap<FName, TypedElementTableHandle> TableNameLookup;
+
+	QueryStore Queries;
+
+	FTypedElementOnDataStorageUpdate OnUpdateDelegate;
+
+	TSharedPtr<FMassEntityManager> ActiveEditorEntityManager;
+};
\ No newline at end of file
diff --git a/Engine/Plugins/Experimental/TypedElementsDataStorage/Source/TypedElementsDataStorage/Private/TypedElementDatabaseCompatibility.cpp b/Engine/Plugins/Experimental/TypedElementsDataStorage/Source/TypedElementsDataStorage/Private/TypedElementDatabaseCompatibility.cpp
new file mode 100644
--- /dev/null
+++ b/Engine/Plugins/Experimental/TypedElementsDataStorage/Source/TypedElementsDataStorage/Private/TypedElementDatabaseCompatibility.cpp
@@ -0,0 +1,149 @@
+// Copyright Epic Games, Inc. All Rights Reserved.
+
+#include "TypedElementDatabaseCompatibility.h"
+
+#include "Editor.h"
+#include "Elements/Columns/TypedElementLabelColumns.h"
+#include "Elements/Columns/TypedElementMiscColumns.h"
+#include "Elements/Columns/TypedElementPackageColumns.h"
+#include "Elements/Columns/TypedElementTransformColumns.h"
+#include "MassActorEditorSubsystem.h"
+#include "MassActorSubsystem.h"
+
+void UTypedElementDatabaseCompatibility::Initialize(ITypedElementDataStorageInterface* StorageInterface)
+{
+	checkf(StorageInterface, TEXT("Typed Element's Database compatibility manager is being initialized with an invalid storage target."));
+
+	Storage = StorageInterface;
+	Prepare();
+
+	StorageInterface->OnUpdate().AddUObject(this, &UTypedElementDatabaseCompatibility::AddPendingActors);
+}
+
+void UTypedElementDatabaseCompatibility::Deinitialize()
+{
+	Reset();
+}
+
+void UTypedElementDatabaseCompatibility::AddCompatibleObject(AActor* Actor)
+{
+	// Registration is delayed for two reasons:
+	//	1. Allows entity creation in a single batch rather than multiple individual additions.
+	//	2. Provides an opportunity to filter out the actors that are created within MASS itself as those will already be registered.
+	ActorsPendingRegistration.Add(Actor);
+}
+
+void UTypedElementDatabaseCompatibility::RemoveCompatibleObject(AActor* Actor)
+{
+	checkf(Storage, TEXT("Removing compatible objects is not supported before Typed Element's Database compatibility manager has been initialized."));
+
+	// If there is no actor subsystem it means that the world has been destroyed, including the MASS instance,
+	// so there's no references to clean up.
+	if (Storage && ActorSubsystem && Storage->IsAvailable())
+	{
+		FMassEntityHandle Entity = ActorSubsystem->GetEntityHandleFromActor(Actor);
+		// If there's no entity it may:
+		//	- have been deleted earlier, e.g. through an explicit delete.
+		//	- be an actor that never had a world assigned and was therefore never registered.
+		//	- have registered with a MASS instance in another world, e.g. one created for PIE.
+		if (Entity.IsValid())
+		{
+			auto ActorStore = Storage->GetColumn<FMassActorFragment>(Entity.AsNumber());
+			if (ActorStore && !ActorStore->IsOwnedByMass()) // Only remove actors that were externally created.
+			{
+				ActorSubsystem->RemoveHandleForActor(Actor);
+				Storage->RemoveRow(Entity.AsNumber());
+			}
+		}
+	}
+}
+
+void UTypedElementDatabaseCompatibility::Prepare()
+{
+	UMassActorEditorSubsystem* MassActorEditorSubsystem = Storage->GetExternalSystem<UMassActorEditorSubsystem>();
+	check(MassActorEditorSubsystem);
+	ActorSubsystem = MassActorEditorSubsystem->GetMutableActorManager().AsShared();
+
+	CreateStandardArchetypes();
+}
+
+void UTypedElementDatabaseCompatibility::Reset()
+{
+	ActorSubsystem = nullptr;
+}
+
+void UTypedElementDatabaseCompatibility::CreateStandardArchetypes()
+{
+	StandardActorTable = Storage->RegisterTable(MakeArrayView(
+		{
+			FMassActorFragment::StaticStruct(),
+			FTypedElementLabelColumn::StaticStruct(),
+			FTypedElementLabelHashColumn::StaticStruct(),
+			FTypedElementPackagePathColumn::StaticStruct(),
+			FTypedElementPackageLoadedPathColumn::StaticStruct()
+		}), FName("Editor_StandardActorTable"));
+
+	StandardActorWithTransformTable = Storage->RegisterTable(StandardActorTable,
+		MakeArrayView(
+		{
+			FTypedElementLabelColumn::StaticStruct(),
+			FTypedElementLocalTransformColumn::StaticStruct()
+		}), FName("Editor_StandardActorWithTransformTable"));
+}
+
+void UTypedElementDatabaseCompatibility::AddPendingActors()
+{
+	UWorld* EditorWorld = GEditor->GetEditorWorldContext().World();
+
+	// Delay processing until the required systems are available by not clearing the pending actor list.
+	if (ActorsPendingRegistration.Num() > 0 && Storage && Storage->IsAvailable() && EditorWorld)
+	{
+		// Filter out the actors that are already registered or already destroyed. 
+		// The most common case for this is actors created from within MASS.
+		TWeakObjectPtr<AActor>* ActorBegin = ActorsPendingRegistration.GetData();
+		TWeakObjectPtr<AActor>* ActorIt = ActorsPendingRegistration.GetData();
+		TWeakObjectPtr<AActor>* ActorEnd = ActorBegin + ActorsPendingRegistration.Num();
+		while (ActorIt != ActorEnd)
+		{
+			if (
+				(*ActorIt).IsValid() && // If not valid or stale then the actor has already been deleted before it could be registered.
+				!ActorSubsystem->GetEntityHandleFromActor((*ActorIt).Get()).IsValid() && // If true, this actor is already registered with the Data Storage.
+				(*ActorIt)->GetWorld() == EditorWorld // Only record actors that are added to the editor to avoid taking PIE into account.
+				) 
+			{
+				++ActorIt;
+			}
+			else
+			{
+				// Don't shrink the registration array as the array will be reused with a variety of different actor counts.
+				// If memory size becomes an issue it's better to resize the array once after this loop rather than within the
+				// loop to avoid many resizes happening.
+				constexpr bool bAllowToShrink = false;
+				ActorsPendingRegistration.RemoveAtSwap(ActorIt - ActorBegin, 1, bAllowToShrink);
+				--ActorEnd;
+			}
+		}
+
+		if (ActorBegin != ActorEnd) // Check if the number of actors to add isn't zero.
+		{
+			// Add the remaining actors to the data storage.
+			ActorIt = ActorsPendingRegistration.GetData();
+			Storage->BatchAddRow(StandardActorTable, ActorsPendingRegistration.Num(), [this, &ActorIt, ActorEnd](TypedElementRowHandle Row)
+				{
+					auto ActorStore = Storage->GetColumn<FMassActorFragment>(Row);
+					checkf(ActorStore, TEXT("Newly created row didn't contain the expected FMassActorFragment."));
+					
+					constexpr bool bIsOwnedByMass = false;
+					ActorStore->SetNoHandleMapUpdate(FMassEntityHandle::FromNumber(Row), ActorIt->Get(), bIsOwnedByMass);
+					ActorSubsystem->SetHandleForActor(ActorIt->Get(), FMassEntityHandle::FromNumber(Row));
+
+					checkf(ActorIt < ActorEnd,
+						TEXT("More (%i) entities were added than were requested (%i)."), ActorEnd - ActorIt, ActorsPendingRegistration.Num());
+					++ActorIt;
+				});
+
+			// Reset the container for next set of actors.
+			ActorsPendingRegistration.Reset();
+		}
+	}
+}
\ No newline at end of file
diff --git a/Engine/Plugins/Experimental/TypedElementsDataStorage/Source/TypedElementsDataStorage/Private/TypedElementDatabaseCompatibility.h b/Engine/Plugins/Experimental/TypedElementsDataStorage/Source/TypedElementsDataStorage/Private/TypedElementDatabaseCompatibility.h
new file mode 100644
--- /dev/null
+++ b/Engine/Plugins/Experimental/TypedElementsDataStorage/Source/TypedElementsDataStorage/Private/TypedElementDatabaseCompatibility.h
@@ -0,0 +1,45 @@
+// Copyright Epic Games, Inc. All Rights Reserved.
+
+#pragma once
+
+#include "Containers/Array.h"
+#include "Elements/Interfaces/TypedElementDataStorageInterface.h"
+#include "Elements/Interfaces/TypedElementDataStorageCompatibilityInterface.h"
+#include "UObject/ObjectMacros.h"
+#include "UObject/WeakObjectPtr.h"
+
+#include "TypedElementDatabaseCompatibility.generated.h"
+
+class AActor;
+class ITypedElementDataStorageInterface;
+struct FMassActorManager;
+
+UCLASS()
+class TYPEDELEMENTSDATASTORAGE_API UTypedElementDatabaseCompatibility
+	: public UObject
+	, public ITypedElementDataStorageCompatibilityInterface
+{
+	GENERATED_BODY()
+public:
+	~UTypedElementDatabaseCompatibility() override = default;
+
+	void Initialize(ITypedElementDataStorageInterface* StorageInterface);
+	void Deinitialize();
+
+	void AddCompatibleObject(AActor* Actor) override;
+	void RemoveCompatibleObject(AActor* Actor) override;
+
+private:
+	void Prepare();
+	void Reset();
+	void CreateStandardArchetypes();
+	
+	void AddPendingActors();
+	
+	TArray<TWeakObjectPtr<AActor>> ActorsPendingRegistration;
+	
+	TypedElementTableHandle StandardActorTable{ TypedElementInvalidTableHandle };
+	TypedElementTableHandle StandardActorWithTransformTable{ TypedElementInvalidTableHandle };
+	ITypedElementDataStorageInterface* Storage{ nullptr };
+	TSharedPtr<FMassActorManager> ActorSubsystem;
+};
\ No newline at end of file
diff --git a/Engine/Plugins/Experimental/TypedElementsDataStorage/Source/TypedElementsDataStorage/Private/TypedElementDatabaseUI.cpp b/Engine/Plugins/Experimental/TypedElementsDataStorage/Source/TypedElementsDataStorage/Private/TypedElementDatabaseUI.cpp
new file mode 100644
--- /dev/null
+++ b/Engine/Plugins/Experimental/TypedElementsDataStorage/Source/TypedElementsDataStorage/Private/TypedElementDatabaseUI.cpp
@@ -0,0 +1,126 @@
+// Copyright Epic Games, Inc. All Rights Reserved.
+
+#include "TypedElementDatabaseUI.h"
+
+#include "Elements/Columns/TypedElementSlateWidgetColumns.h"
+#include "GenericPlatform/GenericPlatformMemory.h"
+
+void UTypedElementDatabaseUi::Initialize(ITypedElementDataStorageInterface* StorageInterface)
+{
+	checkf(StorageInterface, TEXT("Typed Element's Database compatibility manager is being initialized with an invalid storage target."));
+
+	Storage = StorageInterface;
+	CreateStandardArchetypes();
+}
+
+void UTypedElementDatabaseUi::Deinitialize()
+{
+}
+
+void UTypedElementDatabaseUi::RegisterWidgetFactory(FName Purpose, const UScriptStruct* Constructor)
+{
+	checkf(Constructor->IsChildOf(FTypedElementWidgetConstructor::StaticStruct()),
+		TEXT("Attempting to register a Typed Elements widget constructor '%s' that isn't deriving from FTypedElementWidgetConstructor."),
+		*Constructor->GetFullName());
+	WidgetFactoryStructs.Add(Purpose, Constructor);
+}
+
+void UTypedElementDatabaseUi::RegisterWidgetFactory(
+	FName Purpose, const UScriptStruct* ConstructorType, TUniquePtr<FTypedElementWidgetConstructor>&& Constructor)
+{
+	WidgetFactoryInstances.Add(Purpose, { MoveTemp(Constructor), ConstructorType });
+}
+
+void UTypedElementDatabaseUi::ConstructWidgets(FName Purpose, TConstArrayView<TypedElement::ColumnUtils::Argument> Arguments,
+	const WidgetCreatedCallback& ConstructionCallback)
+{
+	// Construct with structs
+	for (TMultiMap<FName, const UScriptStruct*>::TKeyIterator KeyValue = WidgetFactoryStructs.CreateKeyIterator(Purpose); 
+		KeyValue; ++KeyValue)
+	{
+		const UScriptStruct* Target = KeyValue.Value();
+		CreateWidgetInstanceFromDescription(Target, Arguments, ConstructionCallback);
+	}
+
+	// Construct with instances
+	for (TMultiMap<FName, FInstanceConstructor>::TKeyIterator KeyValue = WidgetFactoryInstances.CreateKeyIterator(Purpose); 
+		KeyValue; ++KeyValue)
+	{
+		FTypedElementWidgetConstructor* Source = KeyValue.Value().Constructor.Get();
+		const UScriptStruct* Target = KeyValue.Value().ConstructorType;
+		CreateWidgetInstanceFromInstance(Source, Target, Arguments, ConstructionCallback);
+	}
+}
+
+void UTypedElementDatabaseUi::CreateWidgetInstanceFromDescription(
+	const UScriptStruct* Target,
+	TConstArrayView<TypedElement::ColumnUtils::Argument> Arguments,
+	const WidgetCreatedCallback& ConstructionCallback)
+{
+	FTypedElementWidgetConstructor* Constructor = reinterpret_cast<FTypedElementWidgetConstructor*>(
+		FMemory_Alloca_Aligned(Target->GetStructureSize(), Target->GetMinAlignment()));
+	if (Constructor)
+	{
+		Target->InitializeStruct(Constructor);
+		CreateWidgetInstance(*Constructor, Arguments, ConstructionCallback, Target);
+	}
+	else
+	{
+		checkf(false, TEXT("Remaining stack space is too small to create a Typed Elements widget constructor from a description."));
+	}
+}
+
+void UTypedElementDatabaseUi::CreateWidgetInstanceFromInstance(
+	FTypedElementWidgetConstructor* SourceConstructor,
+	const UScriptStruct* Target,
+	TConstArrayView<TypedElement::ColumnUtils::Argument> Arguments,
+	const WidgetCreatedCallback& ConstructionCallback)
+{
+	FTypedElementWidgetConstructor* Constructor = reinterpret_cast<FTypedElementWidgetConstructor*>(
+		FMemory_Alloca_Aligned(Target->GetStructureSize(), Target->GetMinAlignment()));
+	if (Constructor)
+	{
+		Target->InitializeStruct(Constructor);
+		Target->CopyScriptStruct(Constructor, SourceConstructor);
+		CreateWidgetInstance(*Constructor, Arguments, ConstructionCallback, Target);
+	}
+	else
+	{
+		checkf(false, TEXT("Remaining stack space is too small to create a Typed Elements widget constructor from an instance."));
+	}
+}
+
+void UTypedElementDatabaseUi::CreateWidgetInstance(
+	FTypedElementWidgetConstructor& Constructor, 
+	TConstArrayView<TypedElement::ColumnUtils::Argument> Arguments, 
+	const WidgetCreatedCallback& ConstructionCallback, 
+	const UScriptStruct* Target)
+{
+	TypedElementRowHandle Row = Storage->AddRow(WidgetTable);
+	TSharedPtr<SWidget> Widget = Constructor.Construct(Row, Storage, this, Arguments);
+	if (Widget)
+	{
+		ConstructionCallback(Widget.ToSharedRef(), Row);
+	}
+	else
+	{
+		Storage->RemoveRow(Row);
+	}
+
+	Target->DestroyStruct(&Constructor);
+}
+
+TSharedPtr<SWidget> UTypedElementDatabaseUi::ConstructWidget(TypedElementRowHandle Row, FTypedElementWidgetConstructor& Constructor,
+	TConstArrayView<TypedElement::ColumnUtils::Argument> Arguments)
+{
+	return Constructor.Construct(Row, Storage, this, Arguments);
+}
+
+void UTypedElementDatabaseUi::CreateStandardArchetypes()
+{
+	WidgetTable = Storage->RegisterTable(MakeArrayView(
+		{
+			FTypedElementSlateWidgetReferenceColumn::StaticStruct(),
+			FTypedElementSlateWidgetReferenceDeletesRowTag::StaticStruct()
+		}), FName("Editor_WidgetTable"));
+}
\ No newline at end of file
diff --git a/Engine/Plugins/Experimental/TypedElementsDataStorage/Source/TypedElementsDataStorage/Private/TypedElementDatabaseUI.h b/Engine/Plugins/Experimental/TypedElementsDataStorage/Source/TypedElementsDataStorage/Private/TypedElementDatabaseUI.h
new file mode 100644
--- /dev/null
+++ b/Engine/Plugins/Experimental/TypedElementsDataStorage/Source/TypedElementsDataStorage/Private/TypedElementDatabaseUI.h
@@ -0,0 +1,66 @@
+// Copyright Epic Games, Inc. All Rights Reserved.
+
+#pragma once
+
+#include "Containers/Map.h"
+#include "Elements/Interfaces/TypedElementDataStorageInterface.h"
+#include "Elements/Interfaces/TypedElementDataStorageUiInterface.h"
+#include "UObject/ObjectMacros.h"
+
+#include "TypedElementDatabaseUI.generated.h"
+
+UCLASS()
+class TYPEDELEMENTSDATASTORAGE_API UTypedElementDatabaseUi final
+	: public UObject
+	, public ITypedElementDataStorageUiInterface
+{
+	GENERATED_BODY()
+
+public:
+	~UTypedElementDatabaseUi() override = default;
+
+	void Initialize(ITypedElementDataStorageInterface* StorageInterface);
+	void Deinitialize();
+
+	void RegisterWidgetFactory(FName Purpose, const UScriptStruct* Constructor) override;
+	
+	void ConstructWidgets(FName Purpose, TConstArrayView<TypedElement::ColumnUtils::Argument> Arguments,
+		const WidgetCreatedCallback& ConstructionCallback) override;
+	TSharedPtr<SWidget> ConstructWidget(TypedElementRowHandle Row, FTypedElementWidgetConstructor& Constructor,
+		TConstArrayView<TypedElement::ColumnUtils::Argument> Arguments) override;
+
+protected:
+	void RegisterWidgetFactory(FName Purpose, const UScriptStruct* ConstructorType,
+		TUniquePtr<FTypedElementWidgetConstructor>&& Constructor) override;
+
+private:
+	struct FInstanceConstructor
+	{
+		TUniquePtr<FTypedElementWidgetConstructor> Constructor;
+		const UScriptStruct* ConstructorType;
+	};
+
+	void CreateStandardArchetypes();
+
+	void CreateWidgetInstanceFromDescription(
+		const UScriptStruct* Target,
+		TConstArrayView<TypedElement::ColumnUtils::Argument> Arguments,
+		const WidgetCreatedCallback& ConstructionCallback);
+
+	void CreateWidgetInstanceFromInstance(
+		FTypedElementWidgetConstructor* SourceConstructor,
+		const UScriptStruct* Target,
+		TConstArrayView<TypedElement::ColumnUtils::Argument> Arguments,
+		const WidgetCreatedCallback& ConstructionCallback);
+
+	void CreateWidgetInstance(
+		FTypedElementWidgetConstructor& Constructor,
+		TConstArrayView<TypedElement::ColumnUtils::Argument> Arguments,
+		const WidgetCreatedCallback& ConstructionCallback,
+		const UScriptStruct* Target);
+
+	ITypedElementDataStorageInterface* Storage{ nullptr };
+	TypedElementTableHandle WidgetTable{ TypedElementInvalidTableHandle };
+	TMultiMap<FName, const UScriptStruct*> WidgetFactoryStructs;
+	TMultiMap<FName, FInstanceConstructor> WidgetFactoryInstances;
+};
\ No newline at end of file
diff --git a/Engine/Plugins/Experimental/TypedElementsDataStorage/Source/TypedElementsDataStorage/Private/TypedElementHandleStore.h b/Engine/Plugins/Experimental/TypedElementsDataStorage/Source/TypedElementsDataStorage/Private/TypedElementHandleStore.h
new file mode 100644
--- /dev/null
+++ b/Engine/Plugins/Experimental/TypedElementsDataStorage/Source/TypedElementsDataStorage/Private/TypedElementHandleStore.h
@@ -0,0 +1,111 @@
+// Copyright Epic Games, Inc. All Rights Reserved.
+
+#pragma once
+
+#include "Containers/Array.h"
+#include "Containers/Deque.h"
+
+template<typename DataType, uint32 ReservationSize = 64>
+class TTypedElementHandleStore
+{
+public:
+	struct HandleData
+	{
+		uint32 Index;
+		uint32 Generation;
+	};
+	union Handle
+	{
+		HandleData Data;
+		uint64 Handle;
+	};
+	
+	template<typename... Args>
+	Handle Emplace(Args... Arguments);
+	
+	DataType& Get(Handle Entry);
+	DataType& GetMutable(Handle Entry);
+	const DataType& Get(Handle Entry) const;
+
+	/** Removes the entry at the given handle if alive. This doesn't destroy the entry until the slot is reused. */
+	void Remove(Handle Entry);
+
+	bool IsAlive(Handle Entry) const;
+
+private:
+	TArray<DataType> Data;
+	TArray<uint32> Generations;
+	TDeque<uint32> RecycleBin; // Using a deque to avoid the same slot being constantly reused.
+};
+
+
+// Implementations
+
+template<typename DataType, uint32 ReservationSize>
+template<typename... Args>
+auto TTypedElementHandleStore<DataType, ReservationSize>::Emplace(Args... Arguments) -> Handle
+{
+	if (RecycleBin.IsEmpty())
+	{
+		uint32 NewSize = Data.Num() + ReservationSize;
+		Data.Reserve(NewSize);
+		Generations.Reserve(NewSize);
+
+		for (uint32 Index = Data.Num(); Index < NewSize; ++Index)
+		{
+			Data.AddZeroed();
+			Generations.Add(0);
+			RecycleBin.EmplaceLast(Index);
+		}
+	}
+
+	uint32 Index = RecycleBin.Last();
+	RecycleBin.PopLast();
+	
+	DataType& Entry = Data[Index];
+	new(&Entry) DataType( Forward<Args>(Arguments)... );
+
+	Handle Result;
+	Result.Data.Index = Index;
+	Result.Data.Generation = Generations[Index];
+	return Result;
+}
+
+template<typename DataType, uint32 ReservationSize>
+DataType& TTypedElementHandleStore<DataType, ReservationSize>::Get(Handle Entry)
+{
+	checkf(IsAlive(Entry), TEXT("Attempting to retrieve a dead entry from a Typed Element Handle Store."));
+	return Data[Entry.Data.Index];
+}
+
+template<typename DataType, uint32 ReservationSize>
+DataType& TTypedElementHandleStore<DataType, ReservationSize>::GetMutable(Handle Entry)
+{
+	checkf(IsAlive(Entry), TEXT("Attempting to retrieve a dead entry from a Typed Element Handle Store."));
+	return Data[Entry.Data.Index];
+}
+
+template<typename DataType, uint32 ReservationSize>
+const DataType& TTypedElementHandleStore<DataType, ReservationSize>::Get(Handle Entry) const
+{
+	checkf(IsAlive(Entry), TEXT("Attempting to retrieve a dead entry from a Typed Element Handle Store."));
+	return Data[Entry.Data.Index];
+}
+
+template<typename DataType, uint32 ReservationSize>
+void TTypedElementHandleStore<DataType, ReservationSize>::Remove(Handle Entry)
+{
+	if (IsAlive(Entry))
+	{
+		++Generations[Entry.Data.Index];
+		RecycleBin.EmplaceLast(Entry.Data.Index);
+	}
+}
+
+template<typename DataType, uint32 ReservationSize>
+bool TTypedElementHandleStore<DataType, ReservationSize>::IsAlive(Handle Entry) const
+{
+	return
+		Entry.Data.Index < static_cast<uint32>(Generations.Num()) &&
+		Generations[Entry.Data.Index] == Entry.Data.Generation;
+}
\ No newline at end of file
diff --git a/Engine/Plugins/Experimental/TypedElementsDataStorage/Source/TypedElementsDataStorage/Private/TypedElementSubsystems.cpp b/Engine/Plugins/Experimental/TypedElementsDataStorage/Source/TypedElementsDataStorage/Private/TypedElementSubsystems.cpp
new file mode 100644
--- /dev/null
+++ b/Engine/Plugins/Experimental/TypedElementsDataStorage/Source/TypedElementsDataStorage/Private/TypedElementSubsystems.cpp
@@ -0,0 +1,83 @@
+// Copyright Epic Games, Inc. All Rights Reserved.
+
+#include "TypedElementSubsystems.h"
+
+#include "Elements/Framework/TypedElementRegistry.h"
+
+//
+// UTypedElementDataStorageSubsystem
+//
+
+UTypedElementDataStorageSubsystem::~UTypedElementDataStorageSubsystem()
+{
+	DataStorage = nullptr;
+}
+
+ITypedElementDataStorageInterface* UTypedElementDataStorageSubsystem::Get()
+{
+	if (!DataStorage)
+	{
+		UTypedElementRegistry* Registry = UTypedElementRegistry::GetInstance();
+		checkf(Registry, TEXT("UTypedElementDataStorageSubsystem created before the Typed Elements Registry is available."));
+		DataStorage = Registry->GetMutableDataStorage();
+	}
+	return DataStorage;
+}
+
+const ITypedElementDataStorageInterface* UTypedElementDataStorageSubsystem::Get() const
+{
+	return const_cast<UTypedElementDataStorageSubsystem*>(this)->Get();
+}
+
+
+//
+// UTypedElementDataStorageUiSubsystem
+//
+
+UTypedElementDataStorageUiSubsystem::~UTypedElementDataStorageUiSubsystem()
+{
+	DataStorageUi = nullptr;
+}
+
+ITypedElementDataStorageUiInterface* UTypedElementDataStorageUiSubsystem::Get()
+{
+	if (!DataStorageUi)
+	{
+		UTypedElementRegistry* Registry = UTypedElementRegistry::GetInstance();
+		checkf(Registry, TEXT("UTypedElementDataStorageUiSubsystem created before the Typed Elements Registry is available."));
+		DataStorageUi = Registry->GetMutableDataStorageUi();
+	}
+	return DataStorageUi;
+}
+
+const ITypedElementDataStorageUiInterface* UTypedElementDataStorageUiSubsystem::Get() const
+{
+	return const_cast<UTypedElementDataStorageUiSubsystem*>(this)->Get();
+}
+
+
+//
+// UTypedElementDataStorageCompatibilitySubsystem
+//
+
+UTypedElementDataStorageCompatibilitySubsystem::~UTypedElementDataStorageCompatibilitySubsystem()
+{
+	DataStorageCompatibility = nullptr;
+}
+
+ITypedElementDataStorageCompatibilityInterface* UTypedElementDataStorageCompatibilitySubsystem::Get()
+{
+	if (!DataStorageCompatibility)
+	{
+		UTypedElementRegistry* Registry = UTypedElementRegistry::GetInstance();
+		checkf(Registry, TEXT("UTypedElementDataStorageCompatibilitySubsystem created before the Typed Elements Registry is available."));
+
+		DataStorageCompatibility = Registry->GetMutableDataStorageCompatibility();
+	}
+	return DataStorageCompatibility;
+}
+
+const ITypedElementDataStorageCompatibilityInterface* UTypedElementDataStorageCompatibilitySubsystem::Get() const
+{
+	return const_cast<UTypedElementDataStorageCompatibilitySubsystem*>(this)->Get();
+}
diff --git a/Engine/Plugins/Experimental/TypedElementsDataStorage/Source/TypedElementsDataStorage/Private/TypedElementsDataStorage.cpp b/Engine/Plugins/Experimental/TypedElementsDataStorage/Source/TypedElementsDataStorage/Private/TypedElementsDataStorage.cpp
new file mode 100644
--- /dev/null
+++ b/Engine/Plugins/Experimental/TypedElementsDataStorage/Source/TypedElementsDataStorage/Private/TypedElementsDataStorage.cpp
@@ -0,0 +1,119 @@
+// Copyright Epic Games, Inc. All Rights Reserved.
+
+#include "TypedElementsDataStorage.h"
+
+#include "Elements/Framework/TypedElementRegistry.h"
+#include "MassEntityTypes.h"
+#include "Misc/CoreDelegates.h"
+#include "Templates/IsPolymorphic.h"
+#include "TypedElementDatabase.h"
+#include "TypedElementDatabaseCompatibility.h"
+#include "TypedElementDatabaseUI.h"
+#include "UObject/UObjectGlobals.h"
+
+#define LOCTEXT_NAMESPACE "FTypedElementsDataStorageModule"
+
+// MASS uses DCO in a few places, making it a difficult to consistently register Type Element's Columns and Tags
+// as they may have not been set up to impersonate MASS' Fragments and Tags yet. To get around this, do the 
+// registration during static initialization, which is before DCOs.
+/**
+ * Typed Elements provides base classes for columns and tags. These directly map to fragments and tags in MASS.
+ * To avoid deep and tight coupling between both systems, columns and tags don't directly inhered from MASS, but
+ * are otherwise fully compatible. To allow MASS to do its type safety checks, this class updates the type
+ * information so Typed Elements columns and tags present as MASS fragments and tags from MASS's perspective.
+ */
+struct FTagAndFragmentImpersonators
+{
+	FTagAndFragmentImpersonators()
+	{
+		// Have FTypedElementDataStorageColumn impersonate a FMassFragment, which is the actual data storage when using MASS as a backend.
+		static_assert(sizeof(FTypedElementDataStorageColumn) == sizeof(FMassFragment),
+			"In order for FTypedElementDataStorageColumn to impersonate FMassFragment they need to be identical.");
+		static_assert(!TIsPolymorphic<FMassFragment>::Value,
+			"In order to be able to impersonate FMassFragment it can't have any virtual functions.");
+		static_assert(!TIsPolymorphic<FTypedElementDataStorageColumn>::Value,
+			"In order to be able to use FTypedElementDataStorageColumn to impersonate FMassFragment it can't have any virtual functions.");
+		FTypedElementDataStorageColumn::StaticStruct()->SetSuperStruct(FMassFragment::StaticStruct());
+
+		// Have FTypedElementDataStorageTag impersonate a FMassTag, which is the tag type when using MASS as a backend.
+		static_assert(sizeof(FTypedElementDataStorageTag) == sizeof(FMassTag),
+			"In order for FTypedElementDataStorageTag to impersonate FMassTag they need to be identical.");
+		static_assert(!TIsPolymorphic<FMassTag>::Value,
+			"In order to be able to impersonate FMassTag it can't have any virtual functions.");
+		static_assert(!TIsPolymorphic<FTypedElementDataStorageTag>::Value,
+			"In order to be able to use FTypedElementDataStorageTag to impersonate FMassTag it can't have any virtual functions.");
+		FTypedElementDataStorageTag::StaticStruct()->SetSuperStruct(FMassTag::StaticStruct());
+	}
+};
+static FTagAndFragmentImpersonators TagAndFragmentImpersonators;
+
+void FTypedElementsDataStorageModule::StartupModule()
+{
+	FCoreDelegates::OnAllModuleLoadingPhasesComplete.AddLambda(
+		[this]()
+		{
+			if (!bInitialized)
+			{
+				Database = NewObject<UTypedElementDatabase>();
+				Database->Initialize();
+
+				DatabaseCompatibility = NewObject<UTypedElementDatabaseCompatibility>();
+				DatabaseCompatibility->Initialize(Database.Get());
+
+				DatabaseUi = NewObject<UTypedElementDatabaseUi>();
+				DatabaseUi->Initialize(Database.Get());
+
+				UTypedElementRegistry* Registry = UTypedElementRegistry::GetInstance();
+				checkf(Registry, TEXT(
+					"FTypedElementsDataStorageModule tried to register itself, but there was no Typed Element Registry to register to."));
+				Registry->SetDataStorage(Database.Get());
+				Registry->SetDataStorageCompatibility(DatabaseCompatibility.Get());
+				Registry->SetDataStorageUi(DatabaseUi.Get());
+
+				bInitialized = true;
+			}
+		});
+	FCoreDelegates::OnExit.AddRaw(this, &FTypedElementsDataStorageModule::ShutdownModule);
+}
+
+void FTypedElementsDataStorageModule::ShutdownModule()
+{
+	if (bInitialized)
+	{
+		UTypedElementRegistry* Registry = UTypedElementRegistry::GetInstance();
+		if (Registry) // If the registry has already been destroyed there's no point in clearing the reference.
+		{
+			Registry->SetDataStorage(nullptr);
+			Registry->SetDataStorageCompatibility(nullptr);
+			Registry->SetDataStorageUi(nullptr);
+		}
+
+		if (UObjectInitialized())
+		{
+			DatabaseUi->Deinitialize();
+			DatabaseCompatibility->Deinitialize();
+			Database->Deinitialize();
+		}
+
+		bInitialized = false;
+	}
+}
+
+void FTypedElementsDataStorageModule::AddReferencedObjects(FReferenceCollector& Collector)
+{
+	if (bInitialized)
+	{
+		Collector.AddReferencedObject(Database);
+		Collector.AddReferencedObject(DatabaseCompatibility);
+		Collector.AddReferencedObject(DatabaseUi);
+	}
+}
+
+FString FTypedElementsDataStorageModule::GetReferencerName() const
+{
+	return TEXT("Typed Elements: Data Storage Module");
+}
+
+#undef LOCTEXT_NAMESPACE
+	
+IMPLEMENT_MODULE(FTypedElementsDataStorageModule, TypedElementsDataStorage)
\ No newline at end of file
diff --git a/Engine/Plugins/Experimental/TypedElementsDataStorage/Source/TypedElementsDataStorage/Private/UI/TypedElementSlateWidgetReferenceColumnUpdateProcessor.cpp b/Engine/Plugins/Experimental/TypedElementsDataStorage/Source/TypedElementsDataStorage/Private/UI/TypedElementSlateWidgetReferenceColumnUpdateProcessor.cpp
new file mode 100644
--- /dev/null
+++ b/Engine/Plugins/Experimental/TypedElementsDataStorage/Source/TypedElementsDataStorage/Private/UI/TypedElementSlateWidgetReferenceColumnUpdateProcessor.cpp
@@ -0,0 +1,62 @@
+// Copyright Epic Games, Inc. All Rights Reserved.
+
+#include "UI/TypedElementSlateWidgetReferenceColumnUpdateProcessor.h"
+
+#include "Elements/Columns/TypedElementSlateWidgetColumns.h"
+#include "MassCommonTypes.h"
+#include "MassExecutionContext.h"
+
+UTypedElementSlateWidgetReferenceColumnUpdateProcessor::UTypedElementSlateWidgetReferenceColumnUpdateProcessor()
+	: ColumnRemovalQuery(*this)
+	, RowDeletionQuery(*this)
+{
+	ExecutionFlags = static_cast<int32>(EProcessorExecutionFlags::Editor);
+	ExecutionOrder.ExecuteInGroup = UE::Mass::ProcessorGroupNames::SyncWorldToMass;
+	ProcessingPhase = EMassProcessingPhase::PrePhysics;
+}
+
+void UTypedElementSlateWidgetReferenceColumnUpdateProcessor::ConfigureQueries()
+{
+	ColumnRemovalQuery.AddRequirement(FTypedElementSlateWidgetReferenceColumn::StaticStruct(), EMassFragmentAccess::ReadOnly);
+	ColumnRemovalQuery.AddTagRequirement(*FTypedElementSlateWidgetReferenceDeletesRowTag::StaticStruct(), EMassFragmentPresence::None);
+	
+	RowDeletionQuery.AddRequirement(FTypedElementSlateWidgetReferenceColumn::StaticStruct(), EMassFragmentAccess::ReadOnly);
+	RowDeletionQuery.AddTagRequirement(*FTypedElementSlateWidgetReferenceDeletesRowTag::StaticStruct(), EMassFragmentPresence::All);
+}
+
+void UTypedElementSlateWidgetReferenceColumnUpdateProcessor::Execute(
+	FMassEntityManager& EntityManager, FMassExecutionContext& Context)
+{
+	ColumnRemovalQuery.ForEachEntityChunk(EntityManager, Context, [](FMassExecutionContext& Context)
+		{
+			CheckStatus<false>(Context);
+		});
+
+	RowDeletionQuery.ForEachEntityChunk(EntityManager, Context, [](FMassExecutionContext& Context)
+		{
+			CheckStatus<true>(Context);
+		});
+}
+
+template<bool DeleteRow>
+void UTypedElementSlateWidgetReferenceColumnUpdateProcessor::CheckStatus(FMassExecutionContext& Context)
+{
+	int32 Index = 0;
+	TArrayView<FTypedElementSlateWidgetReferenceColumn> WidgetColumns =
+		Context.GetMutableFragmentView<FTypedElementSlateWidgetReferenceColumn>();
+	for (FTypedElementSlateWidgetReferenceColumn& WidgetColumn : WidgetColumns)
+	{
+		if (!WidgetColumn.Widget.IsValid())
+		{
+			if constexpr (DeleteRow)
+			{
+				Context.Defer().DestroyEntity(Context.GetEntity(Index));
+			}
+			else
+			{
+				Context.Defer().RemoveFragment_RuntimeCheck<FTypedElementSlateWidgetReferenceColumn>(Context.GetEntity(Index));
+			}
+		}
+		++Index;
+	}
+}
diff --git a/Engine/Plugins/Experimental/TypedElementsDataStorage/Source/TypedElementsDataStorage/Private/UI/TypedElementSlateWidgetReferenceColumnUpdateProcessor.h b/Engine/Plugins/Experimental/TypedElementsDataStorage/Source/TypedElementsDataStorage/Private/UI/TypedElementSlateWidgetReferenceColumnUpdateProcessor.h
new file mode 100644
--- /dev/null
+++ b/Engine/Plugins/Experimental/TypedElementsDataStorage/Source/TypedElementsDataStorage/Private/UI/TypedElementSlateWidgetReferenceColumnUpdateProcessor.h
@@ -0,0 +1,33 @@
+// Copyright Epic Games, Inc. All Rights Reserved.
+
+#pragma once
+
+#include "MassObserverProcessor.h"
+#include "MassProcessor.h"
+
+#include "TypedElementSlateWidgetReferenceColumnUpdateProcessor.generated.h"
+
+/**
+ * Processor that checks whether or not a widget still exists. If it has been deleted
+ * then it will remove the column from the Data Storage or deletes the entire row if
+ * the FTypedElementSlateWidgetReferenceDeletesRowTag was found.
+ */
+UCLASS()
+class UTypedElementSlateWidgetReferenceColumnUpdateProcessor : public UMassProcessor
+{
+	GENERATED_BODY()
+
+public:
+	UTypedElementSlateWidgetReferenceColumnUpdateProcessor();
+
+protected:
+	void ConfigureQueries() override;
+	void Execute(FMassEntityManager& EntityManager, FMassExecutionContext& Context) override;
+
+private:
+	template<bool DeleteRow>
+	static void CheckStatus(FMassExecutionContext& Context);
+
+	FMassEntityQuery ColumnRemovalQuery;
+	FMassEntityQuery RowDeletionQuery;
+};
\ No newline at end of file
diff --git a/Engine/Plugins/Experimental/TypedElementsDataStorage/Source/TypedElementsDataStorage/Public/TypedElementSubsystems.h b/Engine/Plugins/Experimental/TypedElementsDataStorage/Source/TypedElementsDataStorage/Public/TypedElementSubsystems.h
new file mode 100644
--- /dev/null
+++ b/Engine/Plugins/Experimental/TypedElementsDataStorage/Source/TypedElementsDataStorage/Public/TypedElementSubsystems.h
@@ -0,0 +1,69 @@
+// Copyright Epic Games, Inc. All Rights Reserved.
+
+#pragma once
+
+#include "EditorSubsystem.h"
+#include "UObject/ObjectMacros.h"
+
+#include "TypedElementSubsystems.generated.h"
+
+class ITypedElementDataStorageInterface;
+class ITypedElementDataStorageUiInterface;
+class ITypedElementDataStorageCompatibilityInterface;
+
+/**
+ * A subsystem to provide alternative access to the Typed Elements Data Storage. This can be used in situations where directly accessing
+ * the Data Storage from the Typed Elements Registry is not recommended, such as for MASS.
+ */
+UCLASS()
+class TYPEDELEMENTSDATASTORAGE_API UTypedElementDataStorageSubsystem : public UEditorSubsystem
+{
+	GENERATED_BODY()
+
+public:
+	~UTypedElementDataStorageSubsystem() override;
+
+	ITypedElementDataStorageInterface* Get();
+	const ITypedElementDataStorageInterface* Get() const;
+
+protected:
+	mutable ITypedElementDataStorageInterface* DataStorage{ nullptr };
+};
+
+/**
+ * A subsystem to provide alternative access to the Typed Elements Data Storage UI. This can be used in situations where directly 
+ * accessing the UI from the Typed Elements Registry is not recommended, such as for MASS.
+ */
+UCLASS()
+class TYPEDELEMENTSDATASTORAGE_API UTypedElementDataStorageUiSubsystem : public UEditorSubsystem
+{
+	GENERATED_BODY()
+
+public:
+	~UTypedElementDataStorageUiSubsystem() override;
+
+	ITypedElementDataStorageUiInterface* Get();
+	const ITypedElementDataStorageUiInterface* Get() const;
+
+protected:
+	mutable ITypedElementDataStorageUiInterface* DataStorageUi{ nullptr };
+};
+
+/**
+ * A subsystem to provide alternative access to the Typed Elements Data Storage Compatiblity. This can be used in situations where directly 
+ * accessing the Compatiblity extension from the Typed Elements Registry is not recommended, such as for MASS.
+ */
+UCLASS()
+class TYPEDELEMENTSDATASTORAGE_API UTypedElementDataStorageCompatibilitySubsystem : public UEditorSubsystem
+{
+	GENERATED_BODY()
+
+public:
+	~UTypedElementDataStorageCompatibilitySubsystem() override;
+
+	ITypedElementDataStorageCompatibilityInterface* Get();
+	const ITypedElementDataStorageCompatibilityInterface* Get() const;
+
+protected:
+	mutable ITypedElementDataStorageCompatibilityInterface* DataStorageCompatibility{ nullptr };
+};
\ No newline at end of file
diff --git a/Engine/Plugins/Experimental/TypedElementsDataStorage/Source/TypedElementsDataStorage/Public/TypedElementsDataStorage.h b/Engine/Plugins/Experimental/TypedElementsDataStorage/Source/TypedElementsDataStorage/Public/TypedElementsDataStorage.h
new file mode 100644
--- /dev/null
+++ b/Engine/Plugins/Experimental/TypedElementsDataStorage/Source/TypedElementsDataStorage/Public/TypedElementsDataStorage.h
@@ -0,0 +1,42 @@
+// Copyright Epic Games, Inc. All Rights Reserved.
+
+#pragma once
+
+#include "CoreMinimal.h"
+#include "Modules/ModuleManager.h"
+#include "UObject/GCObject.h"
+#include "UObject/ObjectPtr.h"
+
+struct FMassEntityManager;
+class FReferenceCollector;
+class ITypedElementDataStorageInterface;
+class UTypedElementDatabase;
+class UTypedElementDatabaseCompatibility;
+class UTypedElementDatabaseUi;
+
+class FTypedElementsDataStorageModule : public IModuleInterface, public FGCObject
+{
+public:
+	~FTypedElementsDataStorageModule() override = default;
+
+	//
+	// IModuleInterface
+	//
+
+	virtual void StartupModule() override;
+	virtual void ShutdownModule() override;
+
+	
+	//
+	// FGCObject
+	//
+	
+	void AddReferencedObjects(FReferenceCollector& Collector) override;
+	FString GetReferencerName() const override;
+
+private:
+	TObjectPtr<UTypedElementDatabase> Database;
+	TObjectPtr<UTypedElementDatabaseCompatibility> DatabaseCompatibility;
+	TObjectPtr<UTypedElementDatabaseUi> DatabaseUi;
+	bool bInitialized{ false };
+};
diff --git a/Engine/Plugins/Experimental/TypedElementsDataStorage/Source/TypedElementsDataStorage/TypedElementsDataStorage.Build.cs b/Engine/Plugins/Experimental/TypedElementsDataStorage/Source/TypedElementsDataStorage/TypedElementsDataStorage.Build.cs
new file mode 100644
--- /dev/null
+++ b/Engine/Plugins/Experimental/TypedElementsDataStorage/Source/TypedElementsDataStorage/TypedElementsDataStorage.Build.cs
@@ -0,0 +1,36 @@
+// Copyright Epic Games, Inc. All Rights Reserved.
+
+using UnrealBuildTool;
+
+public class TypedElementsDataStorage : ModuleRules
+{
+	public TypedElementsDataStorage(ReadOnlyTargetRules Target) : base(Target)
+	{
+		PCHUsage = ModuleRules.PCHUsageMode.UseExplicitOrSharedPCHs;
+
+		if (Target.bBuildEditor)
+		{
+			PublicIncludePaths.AddRange(new string[] {});
+			PrivateIncludePaths.AddRange(new string[] {});
+
+			PublicDependencyModuleNames.AddRange(
+				new string[]
+				{
+					"Core",
+					"CoreUObject",
+					"Engine",
+					"MassActors",
+					"MassEntity",
+					"MassEntityEditor",
+					"MassGameplayEditor",
+					"TypedElementFramework",
+					"EditorSubsystem",
+					"UnrealEd",
+					"StructUtils"
+				});
+
+			PrivateDependencyModuleNames.AddRange(new string[] {});
+			DynamicallyLoadedModuleNames.AddRange(new string[] {});
+		}
+	}
+}
diff --git a/Engine/Plugins/Experimental/TypedElementsDataStorage/Source/TypedElementsDataStorageUI/Private/TypedElementCounterWidgetConstructor.cpp b/Engine/Plugins/Experimental/TypedElementsDataStorage/Source/TypedElementsDataStorageUI/Private/TypedElementCounterWidgetConstructor.cpp
new file mode 100644
--- /dev/null
+++ b/Engine/Plugins/Experimental/TypedElementsDataStorage/Source/TypedElementsDataStorageUI/Private/TypedElementCounterWidgetConstructor.cpp
@@ -0,0 +1,70 @@
+// Copyright Epic Games, Inc. All Rights Reserved.
+
+#include "TypedElementCounterWidgetConstructor.h"
+
+#include "Elements/Columns/TypedElementSlateWidgetColumns.h"
+#include "Elements/Columns/TypedElementValueCacheColumns.h"
+#include "Elements/Framework/TypedElementQueryBuilder.h"
+#include "Framework/Text/TextLayout.h"
+#include "Layout/Margin.h"
+#include "Widgets/DeclarativeSyntaxSupport.h"
+#include "Widgets/Text/STextBlock.h"
+
+#define LOCTEXT_NAMESPACE "TypedElementUI_CounterWidget"
+
+void FTypedElementCounterWidgetConstructor::Register(
+	ITypedElementDataStorageInterface& DataStorage, ITypedElementDataStorageUiInterface& DataStorageUi)
+{
+	using namespace TypedElementQueryBuilder;
+
+	FName Purpose("LevelEditor.StatusBar.ToolBar");
+
+	TUniquePtr<FTypedElementCounterWidgetConstructor> ActorCounter = MakeUnique<FTypedElementCounterWidgetConstructor>();
+	ActorCounter->LabelText = LOCTEXT("ActorCounterStatusBarLabel", "{0} {0}|plural(one=Actor, other=Actors)");
+	ActorCounter->ToolTipText = LOCTEXT(
+		"ActorCounterStatusBarToolTip",
+		"The total number of actors currently in the editor, excluding PIE/SIE and previews.");
+	ActorCounter->Query = DataStorage.RegisterQuery(
+		Count().
+		Where().
+		All("/Script/MassActors.MassActorFragment"_Type).
+		Compile());
+	DataStorageUi.RegisterWidgetFactory(Purpose, MoveTemp(ActorCounter));
+
+	TUniquePtr<FTypedElementCounterWidgetConstructor> WidgetCounter = MakeUnique<FTypedElementCounterWidgetConstructor>();
+	WidgetCounter->LabelText = LOCTEXT("WidgetCounterStatusBarLabel", "{0} {0}|plural(one=Widget, other=Widgets)");
+	WidgetCounter->ToolTipText = LOCTEXT(
+		"WidgetCounterStatusBarToolTip",
+		"The total number of widgets in the editor hosted through the Typed Element's Data Storage.");
+	WidgetCounter->Query = DataStorage.RegisterQuery(
+		Count().
+		Where().
+		All<FTypedElementSlateWidgetReferenceColumn>().
+		Compile());
+	DataStorageUi.RegisterWidgetFactory(Purpose, MoveTemp(WidgetCounter));
+}
+
+TSharedPtr<SWidget> FTypedElementCounterWidgetConstructor::CreateWidget()
+{
+	return SNew(STextBlock)
+		.Text(FText::Format(LabelText, 0))
+		.Margin(FMargin(4.0f, 0.0f))
+		.Justification(ETextJustify::Center);
+}
+
+void FTypedElementCounterWidgetConstructor::AddColumns(
+	ITypedElementDataStorageInterface* DataStorage, TypedElementRowHandle Row, const TSharedPtr<SWidget>& Widget)
+{
+	Super::AddColumns(DataStorage, Row, Widget);
+
+	FTypedElementCounterWidgetColumn* CounterColumn = DataStorage->AddOrGetColumn<FTypedElementCounterWidgetColumn>(Row);
+	checkf(CounterColumn, TEXT("Added a new FTypedElementCounterWidgetColumn to the Typed Elements Data Storage, but didn't get a valid pointer back."));
+	CounterColumn->LabelTextFormatter = LabelText;
+	CounterColumn->Query = Query;
+
+	FTypedElementU32IntValueCacheColumn* CacheColumn = DataStorage->AddOrGetColumn<FTypedElementU32IntValueCacheColumn>(Row);
+	checkf(CacheColumn, TEXT("Added a new FTypedElementUnsigned32BitIntValueCache to the Typed Elements Data Storage, but didn't get a valid pointer back."));
+	CacheColumn->Value = 0;
+}
+
+#undef LOCTEXT_NAMESPACE
\ No newline at end of file
diff --git a/Engine/Plugins/Experimental/TypedElementsDataStorage/Source/TypedElementsDataStorageUI/Private/TypedElementCounterWidgetProcessor.cpp b/Engine/Plugins/Experimental/TypedElementsDataStorage/Source/TypedElementsDataStorageUI/Private/TypedElementCounterWidgetProcessor.cpp
new file mode 100644
--- /dev/null
+++ b/Engine/Plugins/Experimental/TypedElementsDataStorage/Source/TypedElementsDataStorageUI/Private/TypedElementCounterWidgetProcessor.cpp
@@ -0,0 +1,80 @@
+// Copyright Epic Games, Inc. All Rights Reserved.
+
+#include "TypedElementCounterWidgetProcessor.h"
+
+#include "Elements/Columns/TypedElementSlateWidgetColumns.h"
+#include "Elements/Columns/TypedElementValueCacheColumns.h"
+#include "Elements/Framework/TypedElementQueryBuilder.h"
+#include "Elements/Framework/TypedElementRegistry.h"
+#include "Elements/Interfaces/TypedElementDataStorageUiInterface.h"
+#include "TypedElementCounterWidgetConstructor.h"
+#include "TypedElementSubsystems.h"
+#include "MassCommonTypes.h"
+#include "MassExecutionContext.h"
+#include "Widgets/Text/STextBlock.h"
+
+UTypedElementCounterWidgetProcessor::UTypedElementCounterWidgetProcessor()
+	: UiQuery(*this)
+{
+	ExecutionFlags = static_cast<int32>(EProcessorExecutionFlags::Editor);
+	ExecutionOrder.ExecuteInGroup = UE::Mass::ProcessorGroupNames::UpdateWorldFromMass;
+	ProcessingPhase = EMassProcessingPhase::FrameEnd;
+	bRequiresGameThreadExecution = true;
+}
+
+void UTypedElementCounterWidgetProcessor::ConfigureQueries()
+{
+	UiQuery.AddRequirement(FTypedElementSlateWidgetReferenceColumn::StaticStruct(), EMassFragmentAccess::ReadWrite);
+	UiQuery.AddRequirement(FTypedElementU32IntValueCacheColumn::StaticStruct(), EMassFragmentAccess::ReadWrite);
+	UiQuery.AddRequirement(FTypedElementCounterWidgetColumn::StaticStruct(), EMassFragmentAccess::ReadOnly);
+	
+	ProcessorRequirements.AddSubsystemRequirement(UTypedElementDataStorageSubsystem::StaticClass(), EMassFragmentAccess::ReadWrite);
+}
+
+void UTypedElementCounterWidgetProcessor::Execute(FMassEntityManager& EntityManager, FMassExecutionContext& Context)
+{
+	ITypedElementDataStorageInterface* DataInterface = 
+		Context.GetMutableSubsystemChecked<UTypedElementDataStorageSubsystem>().Get();
+	checkf(DataInterface, TEXT(
+		"FTypedElementsDataStorageUiModule tried to process widgets before the "
+		"Typed Elements Data Storage interface is available."));
+
+	UiQuery.ForEachEntityChunk(EntityManager, Context, [DataInterface](FMassExecutionContext& Context)
+		{
+			FTypedElementSlateWidgetReferenceColumn* WidgetIt = 
+				Context.GetMutableFragmentView<FTypedElementSlateWidgetReferenceColumn>().GetData();
+			FTypedElementU32IntValueCacheColumn* CacheIt = 
+				Context.GetMutableFragmentView<FTypedElementU32IntValueCacheColumn>().GetData();
+			FTypedElementCounterWidgetColumn* CounterIt = 
+				Context.GetMutableFragmentView<FTypedElementCounterWidgetColumn>().GetData();
+
+			const int32 NumEntities = Context.GetNumEntities();
+			for (int32 Index = 0; Index < NumEntities; ++Index)
+			{
+				ITypedElementDataStorageInterface::FQueryResult Result = DataInterface->RunQuery(CounterIt->Query);
+				if (Result.Completed == ITypedElementDataStorageInterface::FQueryResult::ECompletion::Fully &&
+					Result.Count != CacheIt->Value)
+				{
+					checkf(WidgetIt->Widget.IsValid(),
+						TEXT("Referenced widget is not valid. A constructed widget may not have been cleaned up. This can also happen if this "
+						"processor is runningin the same phase as the processors responsible for cleaning up old references."));
+
+					TSharedPtr<SWidget> WidgetPointer = WidgetIt->Widget.Pin();
+					checkf(WidgetPointer, TEXT("Widget pointer was null during processing in Counter Widget Processor. This means "
+						"that the processor that cleans up widgets hasn't been run before this processor."));
+					checkf(WidgetPointer->GetType() == STextBlock::StaticWidgetClass().GetWidgetType(),
+						TEXT("Stored widget with FTypedElementCounterWidgetFragment doesn't match type %s, but was a %s."),
+						*(STextBlock::StaticWidgetClass().GetWidgetType().ToString()),
+						*(WidgetPointer->GetTypeAsString()));
+
+					STextBlock* Widget = static_cast<STextBlock*>(WidgetPointer.Get());
+					Widget->SetText(FText::Format(CounterIt->LabelTextFormatter, Result.Count));
+					CacheIt->Value = Result.Count;
+				}
+
+				++WidgetIt;
+				++CacheIt;
+				++CounterIt;
+			}
+		});
+}
diff --git a/Engine/Plugins/Experimental/TypedElementsDataStorage/Source/TypedElementsDataStorageUI/Private/TypedElementCounterWidgetProcessor.h b/Engine/Plugins/Experimental/TypedElementsDataStorage/Source/TypedElementsDataStorageUI/Private/TypedElementCounterWidgetProcessor.h
new file mode 100644
--- /dev/null
+++ b/Engine/Plugins/Experimental/TypedElementsDataStorage/Source/TypedElementsDataStorageUI/Private/TypedElementCounterWidgetProcessor.h
@@ -0,0 +1,25 @@
+// Copyright Epic Games, Inc. All Rights Reserved.
+
+#pragma once
+
+#include "MassEntityTypes.h"
+#include "MassProcessor.h"
+#include "UObject/ObjectMacros.h"
+
+#include "TypedElementCounterWidgetProcessor.generated.h"
+
+/** Runs the count queries and updates the Slate widgets. */
+UCLASS()
+class TYPEDELEMENTSDATASTORAGEUI_API UTypedElementCounterWidgetProcessor : public UMassProcessor
+{
+	GENERATED_BODY()
+
+public:
+	UTypedElementCounterWidgetProcessor();
+
+protected:
+	void ConfigureQueries() override;
+	void Execute(FMassEntityManager& EntityManager, FMassExecutionContext& Context) override;
+
+	FMassEntityQuery UiQuery;
+};
\ No newline at end of file
diff --git a/Engine/Plugins/Experimental/TypedElementsDataStorage/Source/TypedElementsDataStorageUI/Private/TypedElementsDataStorageUI.cpp b/Engine/Plugins/Experimental/TypedElementsDataStorage/Source/TypedElementsDataStorageUI/Private/TypedElementsDataStorageUI.cpp
new file mode 100644
--- /dev/null
+++ b/Engine/Plugins/Experimental/TypedElementsDataStorage/Source/TypedElementsDataStorageUI/Private/TypedElementsDataStorageUI.cpp
@@ -0,0 +1,111 @@
+// Copyright Epic Games, Inc. All Rights Reserved.
+
+#include "TypedElementsDataStorageUI.h"
+
+#include "Elements/Framework/TypedElementRegistry.h"
+#include "Interfaces/IMainFrameModule.h"
+#include "Misc/CoreDelegates.h"
+#include "ToolMenu.h"
+#include "ToolMenus.h"
+#include "TypedElementCounterWidgetConstructor.h"
+#include "UObject/UObjectGlobals.h"
+#include "Widgets/SWindow.h"
+
+#define LOCTEXT_NAMESPACE "FTypedElementsDataStorageModule"
+
+void FTypedElementsDataStorageUiModule::StartupModule()
+{	
+	FCoreDelegates::OnAllModuleLoadingPhasesComplete.AddLambda(
+		[this]()
+		{
+			UTypedElementRegistry* Registry = UTypedElementRegistry::GetInstance();
+			checkf(Registry, TEXT(
+				"FTypedElementsDataStorageUiModule tried to retrieve the UTypedElementRegistry "
+				"before it was available."));
+			if (Registry->AreDataStorageInterfacesSet())
+			{
+				LateInitialization();
+			}
+			else
+			{
+				Registry->OnDataStorageInterfacesSet().AddStatic(&FTypedElementsDataStorageUiModule::LateInitialization);
+			}
+		});
+}
+
+void FTypedElementsDataStorageUiModule::ShutdownModule()
+{
+}
+
+void FTypedElementsDataStorageUiModule::AddReferencedObjects(FReferenceCollector& Collector)
+{
+}
+
+FString FTypedElementsDataStorageUiModule::GetReferencerName() const
+{
+	return TEXT("Typed Elements: Data Storage UI Module");
+}
+
+void FTypedElementsDataStorageUiModule::LateInitialization()
+{
+	RegisterWidgetFactories();
+	IMainFrameModule::Get().OnMainFrameCreationFinished().AddStatic(&FTypedElementsDataStorageUiModule::SetupMainWindowIntegrations);
+}
+
+void FTypedElementsDataStorageUiModule::RegisterWidgetFactories()
+{
+	UTypedElementRegistry* Registry = UTypedElementRegistry::GetInstance();
+	checkf(Registry, TEXT(
+		"FTypedElementsDataStorageUiModule didn't find the UTypedElementRegistry during widget registration when it should be available."));
+
+	ITypedElementDataStorageInterface* DataInterface = Registry->GetMutableDataStorage();
+	checkf(DataInterface, TEXT(
+		"FTypedElementsDataStorageUiModule tried to register widgets before the "
+		"Typed Elements Data Storage interface is available."));
+	ITypedElementDataStorageUiInterface* UiInterface = Registry->GetMutableDataStorageUi();
+	checkf(UiInterface, TEXT(
+		"FTypedElementsDataStorageUiModule tried to register widgets before the "
+		"Typed Elements Data Storage UI interface is available."));
+
+	FTypedElementCounterWidgetConstructor::Register(*DataInterface, *UiInterface);
+}
+
+void FTypedElementsDataStorageUiModule::SetupMainWindowIntegrations(TSharedPtr<SWindow> ParentWindow, bool bIsRunningStartupDialog)
+{
+	UTypedElementRegistry* Registry = UTypedElementRegistry::GetInstance();
+	checkf(Registry, TEXT(
+		"FTypedElementsDataStorageUiModule didn't find the UTypedElementRegistry during main window integration when it should be available."));
+
+	ITypedElementDataStorageUiInterface* UiInterface = Registry->GetMutableDataStorageUi();
+	checkf(UiInterface, TEXT(
+		"FTypedElementsDataStorageUiModule tried to integrate with the main window before the "
+		"Typed Elements Data Storage UI interface is available."));
+	
+	UToolMenu* Menu = UToolMenus::Get()->ExtendMenu("LevelEditor.StatusBar.ToolBar");
+
+	TArray<TSharedRef<SWidget>> Widgets;
+	UiInterface->ConstructWidgets("LevelEditor.StatusBar.ToolBar", {},
+		[&Widgets](const TSharedRef<SWidget>& NewWidget, TypedElementRowHandle Row)
+		{
+			Widgets.Add(NewWidget);
+		});
+
+	if (!Widgets.IsEmpty())
+	{
+		FToolMenuSection& Section = Menu->AddSection("DataStorageSection");
+		int32 WidgetCount = Widgets.Num();
+
+		Section.AddEntry(FToolMenuEntry::InitWidget("DataStorageStatusBarWidget_0", MoveTemp(Widgets[0]), FText::GetEmpty()));
+		for (int32 I = 1; I < WidgetCount; ++I)
+		{
+			Section.AddSeparator(FName(*FString::Format(TEXT("DataStorageStatusBarWidgetDivider_{0}"), { FString::FromInt(I) })));
+			Section.AddEntry(FToolMenuEntry::InitWidget(
+				FName(*FString::Format(TEXT("DataStorageStatusBarWidget_{0}"), { FString::FromInt(I) })),
+				MoveTemp(Widgets[I]), FText::GetEmpty()));
+		}
+	}
+}
+
+#undef LOCTEXT_NAMESPACE
+	
+IMPLEMENT_MODULE(FTypedElementsDataStorageUiModule, TypedElementsDataStorageUI)
\ No newline at end of file
diff --git a/Engine/Plugins/Experimental/TypedElementsDataStorage/Source/TypedElementsDataStorageUI/Public/TypedElementCounterWidgetConstructor.h b/Engine/Plugins/Experimental/TypedElementsDataStorage/Source/TypedElementsDataStorageUI/Public/TypedElementCounterWidgetConstructor.h
new file mode 100644
--- /dev/null
+++ b/Engine/Plugins/Experimental/TypedElementsDataStorage/Source/TypedElementsDataStorageUI/Public/TypedElementCounterWidgetConstructor.h
@@ -0,0 +1,44 @@
+// Copyright Epic Games, Inc. All Rights Reserved.
+
+#pragma once
+
+#include "Elements/Interfaces/TypedElementDataStorageInterface.h"
+#include "Elements/Interfaces/TypedElementDataStorageUiInterface.h"
+#include "Internationalization/Text.h"
+#include "UObject/ObjectMacros.h"
+
+#include "TypedElementCounterWidgetConstructor.generated.h"
+
+/**
+ * Constructor for the counter widget. The counter widget accepts a "count"-query. The query will be periodically
+ * run and the result is written to a textbox widget after it's been formatted using LabelText. An example for 
+ * LabelText is "{0} {0}|plural(one=MyCounter, other=MyCounters)" which will use "MyCounter" if there's exactly one
+ * entry found and otherwise "MyCounters".
+ */
+USTRUCT()
+struct TYPEDELEMENTSDATASTORAGEUI_API FTypedElementCounterWidgetConstructor : public FTypedElementWidgetConstructor
+{
+	GENERATED_BODY()
+
+public:
+	~FTypedElementCounterWidgetConstructor() override = default;
+
+	FText ToolTipText{ NSLOCTEXT("TypedElementUI_CounterWidget", "Tooltip", "Shows the total number found in the editor.") };
+	FText LabelText{ NSLOCTEXT("TypedElementUI_CounterWidget", "Label", "Counted") };
+	TypedElementQueryHandle Query;
+
+	static void Register(ITypedElementDataStorageInterface& DataStorage, ITypedElementDataStorageUiInterface& DataStorageUi);
+
+protected:
+	TSharedPtr<SWidget> CreateWidget() override;
+	void AddColumns(ITypedElementDataStorageInterface* DataStorage, TypedElementRowHandle Row, const TSharedPtr<SWidget>& Widget) override;
+};
+
+USTRUCT()
+struct TYPEDELEMENTSDATASTORAGEUI_API FTypedElementCounterWidgetColumn : public FTypedElementDataStorageColumn
+{
+	GENERATED_BODY()
+
+	FTextFormat LabelTextFormatter;
+	TypedElementQueryHandle Query;
+};
\ No newline at end of file
diff --git a/Engine/Plugins/Experimental/TypedElementsDataStorage/Source/TypedElementsDataStorageUI/Public/TypedElementsDataStorageUI.h b/Engine/Plugins/Experimental/TypedElementsDataStorage/Source/TypedElementsDataStorageUI/Public/TypedElementsDataStorageUI.h
new file mode 100644
--- /dev/null
+++ b/Engine/Plugins/Experimental/TypedElementsDataStorage/Source/TypedElementsDataStorageUI/Public/TypedElementsDataStorageUI.h
@@ -0,0 +1,37 @@
+// Copyright Epic Games, Inc. All Rights Reserved.
+
+#pragma once
+
+#include "CoreMinimal.h"
+#include "Modules/ModuleManager.h"
+#include "UObject/GCObject.h"
+#include "UObject/ObjectPtr.h"
+
+class FReferenceCollector;
+class SWindow;
+
+class FTypedElementsDataStorageUiModule : public IModuleInterface, public FGCObject
+{
+public:
+	~FTypedElementsDataStorageUiModule() override = default;
+
+	//
+	// IModuleInterface
+	//
+
+	virtual void StartupModule() override;
+	virtual void ShutdownModule() override;
+
+	
+	//
+	// FGCObject
+	//
+	
+	void AddReferencedObjects(FReferenceCollector& Collector) override;
+	FString GetReferencerName() const override;
+
+private:
+	static void LateInitialization();
+	static void RegisterWidgetFactories();
+	static void SetupMainWindowIntegrations(TSharedPtr<SWindow> ParentWindow, bool bIsRunningStartupDialog);
+};
diff --git a/Engine/Plugins/Experimental/TypedElementsDataStorage/Source/TypedElementsDataStorageUI/TypedElementsDataStorageUI.Build.cs b/Engine/Plugins/Experimental/TypedElementsDataStorage/Source/TypedElementsDataStorageUI/TypedElementsDataStorageUI.Build.cs
new file mode 100644
--- /dev/null
+++ b/Engine/Plugins/Experimental/TypedElementsDataStorage/Source/TypedElementsDataStorageUI/TypedElementsDataStorageUI.Build.cs
@@ -0,0 +1,42 @@
+// Copyright Epic Games, Inc. All Rights Reserved.
+
+using UnrealBuildTool;
+
+public class TypedElementsDataStorageUI : ModuleRules
+{
+	public TypedElementsDataStorageUI(ReadOnlyTargetRules Target) : base(Target)
+	{
+		PCHUsage = ModuleRules.PCHUsageMode.UseExplicitOrSharedPCHs;
+
+		if (Target.bBuildEditor)
+		{
+			PublicIncludePaths.AddRange(new string[] {});
+			PrivateIncludePaths.AddRange(new string[] {});
+
+			PublicDependencyModuleNames.AddRange(
+				new string[]
+				{
+					"Core",
+					"CoreUObject",
+					"Engine",
+					"EditorFramework",
+					"Slate",
+					"SlateCore",
+					"ToolMenus",
+					"TypedElementFramework",
+					"TypedElementsDataStorage",
+					"EditorSubsystem",
+					"UnrealEd",
+				});
+			PrivateDependencyModuleNames.AddRange(
+				new string[]
+				{
+					"MassActors",
+					"MassEntity",
+					"StructUtils"
+				});
+
+			DynamicallyLoadedModuleNames.AddRange(new string[] {});
+		}
+	}
+}
diff --git a/Engine/Plugins/Experimental/TypedElementsDataStorage/TypedElementsDataStorage.uplugin b/Engine/Plugins/Experimental/TypedElementsDataStorage/TypedElementsDataStorage.uplugin
new file mode 100644
--- /dev/null
+++ b/Engine/Plugins/Experimental/TypedElementsDataStorage/TypedElementsDataStorage.uplugin
@@ -0,0 +1,41 @@
+{
+	"FileVersion": 3,
+	"Version": 1,
+	"VersionName": "0.1",
+	"FriendlyName": "Typed Elements Data Storage",
+	"Description": "A central extendable data storage for editors and their corresponding data with support for viewing and editing through a collection of widgets.",
+	"Category": "Editor",
+	"CreatedBy": "Epic Games, Inc.",
+	"CreatedByURL": "http://epicgames.com",
+	"DocsURL": "",
+	"MarketplaceURL": "",
+	"SupportURL": "",
+	"CanContainContent": true,
+	"IsBetaVersion": false,
+	"IsExperimentalVersion": true,
+	"Installed": false,
+	"Modules": [
+		{
+			"Name": "TypedElementsDataStorage",
+			"Type": "Editor"
+		},
+		{
+			"Name": "TypedElementsDataStorageUI",
+			"Type": "Editor"
+		}
+	],
+	"Plugins": [
+		{
+			"Name": "MassEntity",
+			"Enabled": true
+		},
+		{
+			"Name": "MassGameplay",
+			"Enabled": true
+		},
+		{
+			"Name": "StructUtils",
+			"Enabled": true
+		}
+	]
+}
\ No newline at end of file
diff --git a/Engine/Plugins/TypedElementsDataStorage/Source/TypedElementsDataStorage/Private/Compatibility/TypedElementActorLabelProcessors.cpp b/Engine/Plugins/TypedElementsDataStorage/Source/TypedElementsDataStorage/Private/Compatibility/TypedElementActorLabelProcessors.cpp
deleted file mode 100644
diff --git a/Engine/Plugins/TypedElementsDataStorage/Source/TypedElementsDataStorage/Private/Compatibility/TypedElementActorLabelProcessors.h b/Engine/Plugins/TypedElementsDataStorage/Source/TypedElementsDataStorage/Private/Compatibility/TypedElementActorLabelProcessors.h
deleted file mode 100644
diff --git a/Engine/Plugins/TypedElementsDataStorage/Source/TypedElementsDataStorage/Private/Compatibility/TypedElementActorPackagePathToColumnProcessor.cpp b/Engine/Plugins/TypedElementsDataStorage/Source/TypedElementsDataStorage/Private/Compatibility/TypedElementActorPackagePathToColumnProcessor.cpp
deleted file mode 100644
diff --git a/Engine/Plugins/TypedElementsDataStorage/Source/TypedElementsDataStorage/Private/Compatibility/TypedElementActorPackagePathToColumnProcessor.h b/Engine/Plugins/TypedElementsDataStorage/Source/TypedElementsDataStorage/Private/Compatibility/TypedElementActorPackagePathToColumnProcessor.h
deleted file mode 100644
diff --git a/Engine/Plugins/TypedElementsDataStorage/Source/TypedElementsDataStorage/Private/Compatibility/TypedElementActorTransformProcessors.cpp b/Engine/Plugins/TypedElementsDataStorage/Source/TypedElementsDataStorage/Private/Compatibility/TypedElementActorTransformProcessors.cpp
deleted file mode 100644
diff --git a/Engine/Plugins/TypedElementsDataStorage/Source/TypedElementsDataStorage/Private/Compatibility/TypedElementActorTransformProcessors.h b/Engine/Plugins/TypedElementsDataStorage/Source/TypedElementsDataStorage/Private/Compatibility/TypedElementActorTransformProcessors.h
deleted file mode 100644
diff --git a/Engine/Plugins/TypedElementsDataStorage/Source/TypedElementsDataStorage/Private/Processors/TypedElementMiscProcessors.cpp b/Engine/Plugins/TypedElementsDataStorage/Source/TypedElementsDataStorage/Private/Processors/TypedElementMiscProcessors.cpp
deleted file mode 100644
diff --git a/Engine/Plugins/TypedElementsDataStorage/Source/TypedElementsDataStorage/Private/Processors/TypedElementMiscProcessors.h b/Engine/Plugins/TypedElementsDataStorage/Source/TypedElementsDataStorage/Private/Processors/TypedElementMiscProcessors.h
deleted file mode 100644
diff --git a/Engine/Plugins/TypedElementsDataStorage/Source/TypedElementsDataStorage/Private/TypedElementDatabase.cpp b/Engine/Plugins/TypedElementsDataStorage/Source/TypedElementsDataStorage/Private/TypedElementDatabase.cpp
deleted file mode 100644
diff --git a/Engine/Plugins/TypedElementsDataStorage/Source/TypedElementsDataStorage/Private/TypedElementDatabase.h b/Engine/Plugins/TypedElementsDataStorage/Source/TypedElementsDataStorage/Private/TypedElementDatabase.h
deleted file mode 100644
diff --git a/Engine/Plugins/TypedElementsDataStorage/Source/TypedElementsDataStorage/Private/TypedElementDatabaseCompatibility.cpp b/Engine/Plugins/TypedElementsDataStorage/Source/TypedElementsDataStorage/Private/TypedElementDatabaseCompatibility.cpp
deleted file mode 100644
diff --git a/Engine/Plugins/TypedElementsDataStorage/Source/TypedElementsDataStorage/Private/TypedElementDatabaseCompatibility.h b/Engine/Plugins/TypedElementsDataStorage/Source/TypedElementsDataStorage/Private/TypedElementDatabaseCompatibility.h
deleted file mode 100644
diff --git a/Engine/Plugins/TypedElementsDataStorage/Source/TypedElementsDataStorage/Private/TypedElementDatabaseUI.cpp b/Engine/Plugins/TypedElementsDataStorage/Source/TypedElementsDataStorage/Private/TypedElementDatabaseUI.cpp
deleted file mode 100644
diff --git a/Engine/Plugins/TypedElementsDataStorage/Source/TypedElementsDataStorage/Private/TypedElementDatabaseUI.h b/Engine/Plugins/TypedElementsDataStorage/Source/TypedElementsDataStorage/Private/TypedElementDatabaseUI.h
deleted file mode 100644
diff --git a/Engine/Plugins/TypedElementsDataStorage/Source/TypedElementsDataStorage/Private/TypedElementHandleStore.h b/Engine/Plugins/TypedElementsDataStorage/Source/TypedElementsDataStorage/Private/TypedElementHandleStore.h
deleted file mode 100644
diff --git a/Engine/Plugins/TypedElementsDataStorage/Source/TypedElementsDataStorage/Private/TypedElementSubsystems.cpp b/Engine/Plugins/TypedElementsDataStorage/Source/TypedElementsDataStorage/Private/TypedElementSubsystems.cpp
deleted file mode 100644
diff --git a/Engine/Plugins/TypedElementsDataStorage/Source/TypedElementsDataStorage/Private/TypedElementsDataStorage.cpp b/Engine/Plugins/TypedElementsDataStorage/Source/TypedElementsDataStorage/Private/TypedElementsDataStorage.cpp
deleted file mode 100644
diff --git a/Engine/Plugins/TypedElementsDataStorage/Source/TypedElementsDataStorage/Private/UI/TypedElementSlateWidgetReferenceColumnUpdateProcessor.cpp b/Engine/Plugins/TypedElementsDataStorage/Source/TypedElementsDataStorage/Private/UI/TypedElementSlateWidgetReferenceColumnUpdateProcessor.cpp
deleted file mode 100644
diff --git a/Engine/Plugins/TypedElementsDataStorage/Source/TypedElementsDataStorage/Private/UI/TypedElementSlateWidgetReferenceColumnUpdateProcessor.h b/Engine/Plugins/TypedElementsDataStorage/Source/TypedElementsDataStorage/Private/UI/TypedElementSlateWidgetReferenceColumnUpdateProcessor.h
deleted file mode 100644
diff --git a/Engine/Plugins/TypedElementsDataStorage/Source/TypedElementsDataStorage/Public/TypedElementSubsystems.h b/Engine/Plugins/TypedElementsDataStorage/Source/TypedElementsDataStorage/Public/TypedElementSubsystems.h
deleted file mode 100644
diff --git a/Engine/Plugins/TypedElementsDataStorage/Source/TypedElementsDataStorage/Public/TypedElementsDataStorage.h b/Engine/Plugins/TypedElementsDataStorage/Source/TypedElementsDataStorage/Public/TypedElementsDataStorage.h
deleted file mode 100644
diff --git a/Engine/Plugins/TypedElementsDataStorage/Source/TypedElementsDataStorage/TypedElementsDataStorage.Build.cs b/Engine/Plugins/TypedElementsDataStorage/Source/TypedElementsDataStorage/TypedElementsDataStorage.Build.cs
deleted file mode 100644
diff --git a/Engine/Plugins/TypedElementsDataStorage/Source/TypedElementsDataStorageUI/Private/TypedElementCounterWidgetConstructor.cpp b/Engine/Plugins/TypedElementsDataStorage/Source/TypedElementsDataStorageUI/Private/TypedElementCounterWidgetConstructor.cpp
deleted file mode 100644
diff --git a/Engine/Plugins/TypedElementsDataStorage/Source/TypedElementsDataStorageUI/Private/TypedElementCounterWidgetProcessor.cpp b/Engine/Plugins/TypedElementsDataStorage/Source/TypedElementsDataStorageUI/Private/TypedElementCounterWidgetProcessor.cpp
deleted file mode 100644
diff --git a/Engine/Plugins/TypedElementsDataStorage/Source/TypedElementsDataStorageUI/Private/TypedElementCounterWidgetProcessor.h b/Engine/Plugins/TypedElementsDataStorage/Source/TypedElementsDataStorageUI/Private/TypedElementCounterWidgetProcessor.h
deleted file mode 100644
diff --git a/Engine/Plugins/TypedElementsDataStorage/Source/TypedElementsDataStorageUI/Private/TypedElementsDataStorageUI.cpp b/Engine/Plugins/TypedElementsDataStorage/Source/TypedElementsDataStorageUI/Private/TypedElementsDataStorageUI.cpp
deleted file mode 100644
diff --git a/Engine/Plugins/TypedElementsDataStorage/Source/TypedElementsDataStorageUI/Public/TypedElementCounterWidgetConstructor.h b/Engine/Plugins/TypedElementsDataStorage/Source/TypedElementsDataStorageUI/Public/TypedElementCounterWidgetConstructor.h
deleted file mode 100644
diff --git a/Engine/Plugins/TypedElementsDataStorage/Source/TypedElementsDataStorageUI/Public/TypedElementsDataStorageUI.h b/Engine/Plugins/TypedElementsDataStorage/Source/TypedElementsDataStorageUI/Public/TypedElementsDataStorageUI.h
deleted file mode 100644
diff --git a/Engine/Plugins/TypedElementsDataStorage/Source/TypedElementsDataStorageUI/TypedElementsDataStorageUI.Build.cs b/Engine/Plugins/TypedElementsDataStorage/Source/TypedElementsDataStorageUI/TypedElementsDataStorageUI.Build.cs
deleted file mode 100644
diff --git a/Engine/Plugins/TypedElementsDataStorage/TypedElementsDataStorage.uplugin b/Engine/Plugins/TypedElementsDataStorage/TypedElementsDataStorage.uplugin
deleted file mode 100644
