Added new DCO-based approach to registering Typed Element Data Storage objects and moved the first processor over to use query callbacks.

The following changes were made:
- TypedElementCounterWidgetConstructor was renamed to TypedElementCounterWidget.
- TypedElementCounterWidgetProcessor was removed and the logic moved into a query callback.
- TypedElementCounterWidget now uses a CDO-based factory to register a query with callback and register the widget constructor.
- Added a new tick group called "SyncWidgets".
- The processor member variable in "FTypedElementDatabaseExtendedQuery" is no longer explicitly set to null as this is already the default behavior.
- Removed initialization code in TypedElementsDataStorageUI that's now no longer needed.

(Note that the failed non-unity build is an unrelated issue as confirmed by the build team.)

diff --git a/Engine/Plugins/Experimental/TypedElementsDataStorage/Source/TypedElementsDataStorage/Private/TypedElementDatabase.cpp b/Engine/Plugins/Experimental/TypedElementsDataStorage/Source/TypedElementsDataStorage/Private/TypedElementDatabase.cpp
--- a/Engine/Plugins/Experimental/TypedElementsDataStorage/Source/TypedElementsDataStorage/Private/TypedElementDatabase.cpp
+++ b/Engine/Plugins/Experimental/TypedElementsDataStorage/Source/TypedElementsDataStorage/Private/TypedElementDatabase.cpp
@@ -448,6 +448,11 @@ FName UTypedElementDatabase::GetQueryTickGroupName(EQueryTickGroups Group) const
 			return UE::Mass::ProcessorGroupNames::SyncWorldToMass;
 		case EQueryTickGroups::SyncDataStorageToExternal:
 			return UE::Mass::ProcessorGroupNames::UpdateWorldFromMass;
+		case EQueryTickGroups::SyncWidgets:
+		{
+			static const FName SyncWidgetGroupName(TEXT("SyncWidgets"));
+			return SyncWidgetGroupName;
+		}
 		default:
 			checkf(false, TEXT("EQueryTickGroups value %i can't be translated to a group name by this Data Storage backend."), static_cast<int>(Group));
 			return NAME_None;
diff --git a/Engine/Plugins/Experimental/TypedElementsDataStorage/Source/TypedElementsDataStorage/Private/TypedElementDatabase.h b/Engine/Plugins/Experimental/TypedElementsDataStorage/Source/TypedElementsDataStorage/Private/TypedElementDatabase.h
--- a/Engine/Plugins/Experimental/TypedElementsDataStorage/Source/TypedElementsDataStorage/Private/TypedElementDatabase.h
+++ b/Engine/Plugins/Experimental/TypedElementsDataStorage/Source/TypedElementsDataStorage/Private/TypedElementDatabase.h
@@ -21,7 +21,7 @@ struct FTypedElementDatabaseExtendedQuery
 	FMassEntityQuery NativeQuery;
 	ITypedElementDataStorageInterface::FQueryDescription::EActionType Action{ ITypedElementDataStorageInterface::FQueryDescription::EActionType::None };
 	ITypedElementDataStorageInterface::FQueryDescription::FCallbackData Callback;
-	TStrongObjectPtr<UMassProcessor> Processor{ nullptr };
+	TStrongObjectPtr<UMassProcessor> Processor;
 	bool bSimpleQuery{ false };
 };
 
diff --git a/Engine/Plugins/Experimental/TypedElementsDataStorage/Source/TypedElementsDataStorage/Private/TypedElementsDataStorage.cpp b/Engine/Plugins/Experimental/TypedElementsDataStorage/Source/TypedElementsDataStorage/Private/TypedElementsDataStorage.cpp
--- a/Engine/Plugins/Experimental/TypedElementsDataStorage/Source/TypedElementsDataStorage/Private/TypedElementsDataStorage.cpp
+++ b/Engine/Plugins/Experimental/TypedElementsDataStorage/Source/TypedElementsDataStorage/Private/TypedElementsDataStorage.cpp
@@ -3,6 +3,7 @@
 #include "TypedElementsDataStorage.h"
 
 #include "Elements/Framework/TypedElementRegistry.h"
+#include "Elements/Interfaces/TypedElementDataStorageFactory.h"
 #include "MassEntityTypes.h"
 #include "Misc/CoreDelegates.h"
 #include "Templates/IsPolymorphic.h"
@@ -63,6 +64,7 @@ void FTypedElementsDataStorageModule::StartupModule()
 				DatabaseUi = NewObject<UTypedElementDatabaseUi>();
 				DatabaseUi->Initialize(Database.Get());
 
+				// Register the various database instances.
 				UTypedElementRegistry* Registry = UTypedElementRegistry::GetInstance();
 				checkf(Registry, TEXT(
 					"FTypedElementsDataStorageModule tried to register itself, but there was no Typed Element Registry to register to."));
@@ -70,6 +72,21 @@ void FTypedElementsDataStorageModule::StartupModule()
 				Registry->SetDataStorageCompatibility(DatabaseCompatibility.Get());
 				Registry->SetDataStorageUi(DatabaseUi.Get());
 
+				// Allow any factories to register their content.
+				TArray<UClass*> FactoryClasses;
+				GetDerivedClasses(UTypedElementDataStorageFactory::StaticClass(), FactoryClasses);
+				for (UClass* Factory : FactoryClasses)
+				{
+					if (Factory->HasAnyClassFlags(CLASS_Abstract))
+					{
+						continue;
+					}
+
+					UTypedElementDataStorageFactory* FactoryCDO = GetMutableDefault<UTypedElementDataStorageFactory>(Factory);
+					FactoryCDO->RegisterQueries(*Database);
+					FactoryCDO->RegisterWidgetConstructor(*Database, *DatabaseUi);
+				}
+
 				bInitialized = true;
 			}
 		});
diff --git a/Engine/Plugins/Experimental/TypedElementsDataStorage/Source/TypedElementsDataStorage/Public/TypedElementsDataStorage.h b/Engine/Plugins/Experimental/TypedElementsDataStorage/Source/TypedElementsDataStorage/Public/TypedElementsDataStorage.h
--- a/Engine/Plugins/Experimental/TypedElementsDataStorage/Source/TypedElementsDataStorage/Public/TypedElementsDataStorage.h
+++ b/Engine/Plugins/Experimental/TypedElementsDataStorage/Source/TypedElementsDataStorage/Public/TypedElementsDataStorage.h
@@ -7,9 +7,7 @@
 #include "UObject/GCObject.h"
 #include "UObject/ObjectPtr.h"
 
-struct FMassEntityManager;
 class FReferenceCollector;
-class ITypedElementDataStorageInterface;
 class UTypedElementDatabase;
 class UTypedElementDatabaseCompatibility;
 class UTypedElementDatabaseUi;
diff --git a/Engine/Plugins/Experimental/TypedElementsDataStorage/Source/TypedElementsDataStorageUI/Private/TypedElementCounterWidget.cpp b/Engine/Plugins/Experimental/TypedElementsDataStorage/Source/TypedElementsDataStorageUI/Private/TypedElementCounterWidget.cpp
new file mode 100644
--- /dev/null
+++ b/Engine/Plugins/Experimental/TypedElementsDataStorage/Source/TypedElementsDataStorageUI/Private/TypedElementCounterWidget.cpp
@@ -0,0 +1,120 @@
+// Copyright Epic Games, Inc. All Rights Reserved.
+
+#include "TypedElementCounterWidget.h"
+
+#include "Elements/Columns/TypedElementSlateWidgetColumns.h"
+#include "Elements/Columns/TypedElementValueCacheColumns.h"
+#include "Elements/Framework/TypedElementQueryBuilder.h"
+#include "Framework/Text/TextLayout.h"
+#include "Layout/Margin.h"
+#include "Widgets/DeclarativeSyntaxSupport.h"
+#include "Widgets/Text/STextBlock.h"
+#include "TypedElementSubsystems.h"
+
+#define LOCTEXT_NAMESPACE "TypedElementUI_CounterWidget"
+
+
+//
+// UTypedElementCounterWidgetFactory
+//
+
+void UTypedElementCounterWidgetFactory::RegisterQueries(ITypedElementDataStorageInterface& DataStorage) const
+{
+	using namespace TypedElementQueryBuilder;
+	using DSI = ITypedElementDataStorageInterface;
+
+	DataStorage.RegisterQuery(Select(TEXT("Sync Counter Widgets"), 
+		FProcessor(DSI::EQueryTickPhase::FrameEnd, DataStorage.GetQueryTickGroupName(DSI::EQueryTickGroups::SyncWidgets)).ForceToGameThread(true),
+		[](
+			FCachedQueryContext<UTypedElementDataStorageSubsystem>& Context, 
+			FTypedElementSlateWidgetReferenceColumn& Widget,
+			FTypedElementU32IntValueCacheColumn& Comparison, 
+			const FTypedElementCounterWidgetColumn& Counter
+		)
+		{
+			UTypedElementDataStorageSubsystem* Subsystem = Context.GetCachedMutableSubsystem<UTypedElementDataStorageSubsystem>();
+			DSI* DataInterface = Subsystem->Get();
+			checkf(DataInterface, TEXT("FTypedElementsDataStorageUiModule tried to process widgets before the "
+				"Typed Elements Data Storage interface is available."));
+
+			DSI::FQueryResult Result = DataInterface->RunQuery(Counter.Query);
+			if (Result.Completed == DSI::FQueryResult::ECompletion::Fully && Result.Count != Comparison.Value)
+			{
+				TSharedPtr<SWidget> WidgetPointer = Widget.Widget.Pin();
+				checkf(WidgetPointer, TEXT("Referenced widget is not valid. A constructed widget may not have been cleaned up. This can "
+					"also happen if this processor is running in the same phase as the processors responsible for cleaning up old "
+					"references."));
+				checkf(WidgetPointer->GetType() == STextBlock::StaticWidgetClass().GetWidgetType(),
+					TEXT("Stored widget with FTypedElementCounterWidgetFragment doesn't match type %s, but was a %s."),
+					*(STextBlock::StaticWidgetClass().GetWidgetType().ToString()),
+					*(WidgetPointer->GetTypeAsString()));
+
+				STextBlock* WidgetInstance = static_cast<STextBlock*>(WidgetPointer.Get());
+				WidgetInstance->SetText(FText::Format(Counter.LabelTextFormatter, Result.Count));
+				Comparison.Value = Result.Count;
+			}
+		}
+	).Compile());
+}
+
+void UTypedElementCounterWidgetFactory::RegisterWidgetConstructor(ITypedElementDataStorageInterface& DataStorage,
+	ITypedElementDataStorageUiInterface& DataStorageUi) const
+{
+	using namespace TypedElementQueryBuilder;
+
+	FName Purpose("LevelEditor.StatusBar.ToolBar");
+
+	TUniquePtr<FTypedElementCounterWidgetConstructor> ActorCounter = MakeUnique<FTypedElementCounterWidgetConstructor>();
+	ActorCounter->LabelText = LOCTEXT("ActorCounterStatusBarLabel", "{0} {0}|plural(one=Actor, other=Actors)");
+	ActorCounter->ToolTipText = LOCTEXT(
+		"ActorCounterStatusBarToolTip",
+		"The total number of actors currently in the editor, excluding PIE/SIE and previews.");
+	ActorCounter->Query = DataStorage.RegisterQuery(
+		Count().
+		Where().
+		All("/Script/MassActors.MassActorFragment"_Type).
+		Compile());
+	DataStorageUi.RegisterWidgetFactory(Purpose, MoveTemp(ActorCounter));
+
+	TUniquePtr<FTypedElementCounterWidgetConstructor> WidgetCounter = MakeUnique<FTypedElementCounterWidgetConstructor>();
+	WidgetCounter->LabelText = LOCTEXT("WidgetCounterStatusBarLabel", "{0} {0}|plural(one=Widget, other=Widgets)");
+	WidgetCounter->ToolTipText = LOCTEXT(
+		"WidgetCounterStatusBarToolTip",
+		"The total number of widgets in the editor hosted through the Typed Element's Data Storage.");
+	WidgetCounter->Query = DataStorage.RegisterQuery(
+		Count().
+		Where().
+		All<FTypedElementSlateWidgetReferenceColumn>().
+		Compile());
+	DataStorageUi.RegisterWidgetFactory(Purpose, MoveTemp(WidgetCounter));
+}
+
+
+//
+// FTypedElementCounterWidgetConstructor
+//
+
+TSharedPtr<SWidget> FTypedElementCounterWidgetConstructor::CreateWidget()
+{
+	return SNew(STextBlock)
+		.Text(FText::Format(LabelText, 0))
+		.Margin(FMargin(4.0f, 0.0f))
+		.Justification(ETextJustify::Center);
+}
+
+void FTypedElementCounterWidgetConstructor::AddColumns(
+	ITypedElementDataStorageInterface* DataStorage, TypedElementRowHandle Row, const TSharedPtr<SWidget>& Widget)
+{
+	Super::AddColumns(DataStorage, Row, Widget);
+
+	FTypedElementCounterWidgetColumn* CounterColumn = DataStorage->AddOrGetColumn<FTypedElementCounterWidgetColumn>(Row);
+	checkf(CounterColumn, TEXT("Added a new FTypedElementCounterWidgetColumn to the Typed Elements Data Storage, but didn't get a valid pointer back."));
+	CounterColumn->LabelTextFormatter = LabelText;
+	CounterColumn->Query = Query;
+
+	FTypedElementU32IntValueCacheColumn* CacheColumn = DataStorage->AddOrGetColumn<FTypedElementU32IntValueCacheColumn>(Row);
+	checkf(CacheColumn, TEXT("Added a new FTypedElementUnsigned32BitIntValueCache to the Typed Elements Data Storage, but didn't get a valid pointer back."));
+	CacheColumn->Value = 0;
+}
+
+#undef LOCTEXT_NAMESPACE
\ No newline at end of file
diff --git a/Engine/Plugins/Experimental/TypedElementsDataStorage/Source/TypedElementsDataStorageUI/Private/TypedElementCounterWidgetConstructor.cpp b/Engine/Plugins/Experimental/TypedElementsDataStorage/Source/TypedElementsDataStorageUI/Private/TypedElementCounterWidgetConstructor.cpp
deleted file mode 100644
diff --git a/Engine/Plugins/Experimental/TypedElementsDataStorage/Source/TypedElementsDataStorageUI/Private/TypedElementCounterWidgetProcessor.cpp b/Engine/Plugins/Experimental/TypedElementsDataStorage/Source/TypedElementsDataStorageUI/Private/TypedElementCounterWidgetProcessor.cpp
deleted file mode 100644
diff --git a/Engine/Plugins/Experimental/TypedElementsDataStorage/Source/TypedElementsDataStorageUI/Private/TypedElementCounterWidgetProcessor.h b/Engine/Plugins/Experimental/TypedElementsDataStorage/Source/TypedElementsDataStorageUI/Private/TypedElementCounterWidgetProcessor.h
deleted file mode 100644
diff --git a/Engine/Plugins/Experimental/TypedElementsDataStorage/Source/TypedElementsDataStorageUI/Private/TypedElementsDataStorageUI.cpp b/Engine/Plugins/Experimental/TypedElementsDataStorage/Source/TypedElementsDataStorageUI/Private/TypedElementsDataStorageUI.cpp
--- a/Engine/Plugins/Experimental/TypedElementsDataStorage/Source/TypedElementsDataStorageUI/Private/TypedElementsDataStorageUI.cpp
+++ b/Engine/Plugins/Experimental/TypedElementsDataStorage/Source/TypedElementsDataStorageUI/Private/TypedElementsDataStorageUI.cpp
@@ -7,30 +7,14 @@
 #include "Misc/CoreDelegates.h"
 #include "ToolMenu.h"
 #include "ToolMenus.h"
-#include "TypedElementCounterWidgetConstructor.h"
 #include "UObject/UObjectGlobals.h"
 #include "Widgets/SWindow.h"
 
 #define LOCTEXT_NAMESPACE "FTypedElementsDataStorageModule"
 
 void FTypedElementsDataStorageUiModule::StartupModule()
-{	
-	FCoreDelegates::OnAllModuleLoadingPhasesComplete.AddLambda(
-		[this]()
-		{
-			UTypedElementRegistry* Registry = UTypedElementRegistry::GetInstance();
-			checkf(Registry, TEXT(
-				"FTypedElementsDataStorageUiModule tried to retrieve the UTypedElementRegistry "
-				"before it was available."));
-			if (Registry->AreDataStorageInterfacesSet())
-			{
-				LateInitialization();
-			}
-			else
-			{
-				Registry->OnDataStorageInterfacesSet().AddStatic(&FTypedElementsDataStorageUiModule::LateInitialization);
-			}
-		});
+{
+	IMainFrameModule::Get().OnMainFrameCreationFinished().AddStatic(&FTypedElementsDataStorageUiModule::SetupMainWindowIntegrations);
 }
 
 void FTypedElementsDataStorageUiModule::ShutdownModule()
@@ -46,30 +30,6 @@ FString FTypedElementsDataStorageUiModule::GetReferencerName() const
 	return TEXT("Typed Elements: Data Storage UI Module");
 }
 
-void FTypedElementsDataStorageUiModule::LateInitialization()
-{
-	RegisterWidgetFactories();
-	IMainFrameModule::Get().OnMainFrameCreationFinished().AddStatic(&FTypedElementsDataStorageUiModule::SetupMainWindowIntegrations);
-}
-
-void FTypedElementsDataStorageUiModule::RegisterWidgetFactories()
-{
-	UTypedElementRegistry* Registry = UTypedElementRegistry::GetInstance();
-	checkf(Registry, TEXT(
-		"FTypedElementsDataStorageUiModule didn't find the UTypedElementRegistry during widget registration when it should be available."));
-
-	ITypedElementDataStorageInterface* DataInterface = Registry->GetMutableDataStorage();
-	checkf(DataInterface, TEXT(
-		"FTypedElementsDataStorageUiModule tried to register widgets before the "
-		"Typed Elements Data Storage interface is available."));
-	ITypedElementDataStorageUiInterface* UiInterface = Registry->GetMutableDataStorageUi();
-	checkf(UiInterface, TEXT(
-		"FTypedElementsDataStorageUiModule tried to register widgets before the "
-		"Typed Elements Data Storage UI interface is available."));
-
-	FTypedElementCounterWidgetConstructor::Register(*DataInterface, *UiInterface);
-}
-
 void FTypedElementsDataStorageUiModule::SetupMainWindowIntegrations(TSharedPtr<SWindow> ParentWindow, bool bIsRunningStartupDialog)
 {
 	UTypedElementRegistry* Registry = UTypedElementRegistry::GetInstance();
diff --git a/Engine/Plugins/Experimental/TypedElementsDataStorage/Source/TypedElementsDataStorageUI/Public/TypedElementCounterWidget.h b/Engine/Plugins/Experimental/TypedElementsDataStorage/Source/TypedElementsDataStorageUI/Public/TypedElementCounterWidget.h
new file mode 100644
--- /dev/null
+++ b/Engine/Plugins/Experimental/TypedElementsDataStorage/Source/TypedElementsDataStorageUI/Public/TypedElementCounterWidget.h
@@ -0,0 +1,56 @@
+// Copyright Epic Games, Inc. All Rights Reserved.
+
+#pragma once
+
+#include "Elements/Interfaces/TypedElementDataStorageFactory.h"
+#include "Elements/Interfaces/TypedElementDataStorageInterface.h"
+#include "Elements/Interfaces/TypedElementDataStorageUiInterface.h"
+#include "Internationalization/Text.h"
+#include "UObject/ObjectMacros.h"
+
+#include "TypedElementCounterWidget.generated.h"
+
+UCLASS()
+class TYPEDELEMENTSDATASTORAGEUI_API UTypedElementCounterWidgetFactory : public UTypedElementDataStorageFactory
+{
+	GENERATED_BODY()
+
+public:
+	~UTypedElementCounterWidgetFactory() override = default;
+
+	void RegisterQueries(ITypedElementDataStorageInterface& DataStorage) const override;
+	void RegisterWidgetConstructor(ITypedElementDataStorageInterface& DataStorage,
+		ITypedElementDataStorageUiInterface& DataStorageUi) const override;
+};
+
+/**
+ * Constructor for the counter widget. The counter widget accepts a "count"-query. The query will be periodically
+ * run and the result is written to a textbox widget after it's been formatted using LabelText. An example for 
+ * LabelText is "{0} {0}|plural(one=MyCounter, other=MyCounters)" which will use "MyCounter" if there's exactly one
+ * entry found and otherwise "MyCounters".
+ */
+USTRUCT()
+struct TYPEDELEMENTSDATASTORAGEUI_API FTypedElementCounterWidgetConstructor : public FTypedElementWidgetConstructor
+{
+	GENERATED_BODY()
+
+public:
+	~FTypedElementCounterWidgetConstructor() override = default;
+
+	FText ToolTipText{ NSLOCTEXT("TypedElementUI_CounterWidget", "Tooltip", "Shows the total number found in the editor.") };
+	FText LabelText{ NSLOCTEXT("TypedElementUI_CounterWidget", "Label", "Counted") };
+	TypedElementQueryHandle Query;
+
+protected:
+	TSharedPtr<SWidget> CreateWidget() override;
+	void AddColumns(ITypedElementDataStorageInterface* DataStorage, TypedElementRowHandle Row, const TSharedPtr<SWidget>& Widget) override;
+};
+
+USTRUCT()
+struct TYPEDELEMENTSDATASTORAGEUI_API FTypedElementCounterWidgetColumn : public FTypedElementDataStorageColumn
+{
+	GENERATED_BODY()
+
+	FTextFormat LabelTextFormatter;
+	TypedElementQueryHandle Query;
+};
\ No newline at end of file
diff --git a/Engine/Plugins/Experimental/TypedElementsDataStorage/Source/TypedElementsDataStorageUI/Public/TypedElementCounterWidgetConstructor.h b/Engine/Plugins/Experimental/TypedElementsDataStorage/Source/TypedElementsDataStorageUI/Public/TypedElementCounterWidgetConstructor.h
deleted file mode 100644
diff --git a/Engine/Plugins/Experimental/TypedElementsDataStorage/Source/TypedElementsDataStorageUI/Public/TypedElementsDataStorageUI.h b/Engine/Plugins/Experimental/TypedElementsDataStorage/Source/TypedElementsDataStorageUI/Public/TypedElementsDataStorageUI.h
--- a/Engine/Plugins/Experimental/TypedElementsDataStorage/Source/TypedElementsDataStorageUI/Public/TypedElementsDataStorageUI.h
+++ b/Engine/Plugins/Experimental/TypedElementsDataStorage/Source/TypedElementsDataStorageUI/Public/TypedElementsDataStorageUI.h
@@ -31,7 +31,5 @@ public:
 	FString GetReferencerName() const override;
 
 private:
-	static void LateInitialization();
-	static void RegisterWidgetFactories();
 	static void SetupMainWindowIntegrations(TSharedPtr<SWindow> ParentWindow, bool bIsRunningStartupDialog);
 };
diff --git a/Engine/Source/Runtime/TypedElementFramework/Public/Elements/Framework/TypedElementQueryBuilder.inl b/Engine/Source/Runtime/TypedElementFramework/Public/Elements/Framework/TypedElementQueryBuilder.inl
--- a/Engine/Source/Runtime/TypedElementFramework/Public/Elements/Framework/TypedElementQueryBuilder.inl
+++ b/Engine/Source/Runtime/TypedElementFramework/Public/Elements/Framework/TypedElementQueryBuilder.inl
@@ -13,7 +13,7 @@ namespace TypedElementQueryBuilder
 		template<typename Target, typename ArgsCurrent, typename... ArgsRemainder>
 		constexpr uint32 GetVarArgIndex()
 		{
-			if constexpr (TIsSame<Target, ArgsCurrent>::Value)
+			if constexpr (std::is_same_v<Target, ArgsCurrent>)
 			{
 				return 0;
 			}
@@ -151,7 +151,7 @@ namespace TypedElementQueryBuilder
 			uint32 ColumnSize[] = { sizeof(Args)... };																		\
 																															\
 			using BaseContextType = typename TRemoveReference<ContextType>::Type;											\
-			if constexpr(TIsSame<BaseContextType, QueryContext>::Value)														\
+			if constexpr(std::is_same_v<BaseContextType, QueryContext>)														\
 			{																												\
 				const uint32 RowCount = Context.GetRowCount();																\
 				for (uint32 Index = 0; Index < RowCount; ++Index)															\
@@ -302,7 +302,7 @@ namespace TypedElementQueryBuilder
 		constexpr bool IsValidContextType()
 		{
 			using BaseContextType = typename std::remove_const_t<typename TRemoveReference<Context>::Type>;
-			if constexpr (TIsSame<BaseContextType, QueryContext>::Value)
+			if constexpr (std::is_same_v<BaseContextType, QueryContext>)
 			{
 				return true;
 			}
diff --git a/Engine/Source/Runtime/TypedElementFramework/Public/Elements/Interfaces/TypedElementDataStorageFactory.h b/Engine/Source/Runtime/TypedElementFramework/Public/Elements/Interfaces/TypedElementDataStorageFactory.h
new file mode 100644
--- /dev/null
+++ b/Engine/Source/Runtime/TypedElementFramework/Public/Elements/Interfaces/TypedElementDataStorageFactory.h
@@ -0,0 +1,25 @@
+// Copyright Epic Games, Inc. All Rights Reserved.
+
+#pragma once
+
+#include "TypedElementDataStorageFactory.generated.h"
+
+class ITypedElementDataStorageInterface;
+class ITypedElementDataStorageUiInterface;
+
+/**
+ * Base class that can be used to register various elements, such as queries and widgets, with
+ * the Typed Elements Data Storage.
+ */
+UCLASS()
+class TYPEDELEMENTFRAMEWORK_API UTypedElementDataStorageFactory : public UObject
+{
+	GENERATED_BODY()
+
+public:
+	~UTypedElementDataStorageFactory() override = default;
+
+	virtual void RegisterQueries(ITypedElementDataStorageInterface& DataStorage) const {}
+	virtual void RegisterWidgetConstructor(ITypedElementDataStorageInterface& DataStorage,
+		ITypedElementDataStorageUiInterface& DataStorageUi) const {}
+};
\ No newline at end of file
diff --git a/Engine/Source/Runtime/TypedElementFramework/Public/Elements/Interfaces/TypedElementDataStorageInterface.h b/Engine/Source/Runtime/TypedElementFramework/Public/Elements/Interfaces/TypedElementDataStorageInterface.h
--- a/Engine/Source/Runtime/TypedElementFramework/Public/Elements/Interfaces/TypedElementDataStorageInterface.h
+++ b/Engine/Source/Runtime/TypedElementFramework/Public/Elements/Interfaces/TypedElementDataStorageInterface.h
@@ -169,7 +169,11 @@ public:
 		 * The group for processors that need to sync data from the Data Storage to external sources such as subsystems 
 		 * or the world into. These typically run late in a phase.
 		 */
-		SyncDataStorageToExternal
+		SyncDataStorageToExternal,
+		/**
+		 * Processors grouped under this name will sync data to/from widgets.
+		 */
+		SyncWidgets
 	};
 
 	enum class EQueryCallbackType : uint8
