The remaining MASS Processors used by the Typed Elements Data Storage have been converted to query callbacks.

diff --git a/Engine/Plugins/Experimental/TypedElementsDataStorage/Source/TypedElementsDataStorage/Private/Compatibility/TypedElementActorLabelProcessors.cpp b/Engine/Plugins/Experimental/TypedElementsDataStorage/Source/TypedElementsDataStorage/Private/Compatibility/TypedElementActorLabelProcessors.cpp
--- a/Engine/Plugins/Experimental/TypedElementsDataStorage/Source/TypedElementsDataStorage/Private/Compatibility/TypedElementActorLabelProcessors.cpp
+++ b/Engine/Plugins/Experimental/TypedElementsDataStorage/Source/TypedElementsDataStorage/Private/Compatibility/TypedElementActorLabelProcessors.cpp
@@ -4,106 +4,70 @@
 
 #include "Elements/Columns/TypedElementMiscColumns.h"
 #include "Elements/Columns/TypedElementLabelColumns.h"
+#include "Elements/Framework/TypedElementQueryBuilder.h"
 #include "Hash/CityHash.h"
 #include "MassActorSubsystem.h"
-#include "MassExecutionContext.h"
 
-/**
- * UTypedElementActorLabelToColumnProcessor
- */
 
-UTypedElementActorLabelToColumnProcessor::UTypedElementActorLabelToColumnProcessor()
-	: Query(*this)
+void UTypedElementActorLabelFactory::RegisterQueries(ITypedElementDataStorageInterface& DataStorage) const
 {
-	ExecutionFlags = static_cast<int32>(EProcessorExecutionFlags::Editor);
-	ExecutionOrder.ExecuteInGroup = UE::Mass::ProcessorGroupNames::SyncWorldToMass;
-	ProcessingPhase = EMassProcessingPhase::PrePhysics;
-	bRequiresGameThreadExecution = true;
+	RegisterActorLabelToColumnQuery(DataStorage);
+	RegisterLabelColumnToActorQuery(DataStorage);
 }
 
-void UTypedElementActorLabelToColumnProcessor::ConfigureQueries()
+void UTypedElementActorLabelFactory::RegisterActorLabelToColumnQuery(ITypedElementDataStorageInterface& DataStorage) const
 {
-	Query.AddRequirement(FMassActorFragment::StaticStruct(), EMassFragmentAccess::ReadOnly);
-	Query.AddRequirement(FTypedElementLabelColumn::StaticStruct(), EMassFragmentAccess::ReadWrite);
-	Query.AddRequirement(FTypedElementLabelHashColumn::StaticStruct(), EMassFragmentAccess::ReadWrite);
-}
+	using namespace TypedElementQueryBuilder;
+	using DSI = ITypedElementDataStorageInterface;
 
-void UTypedElementActorLabelToColumnProcessor::Execute(FMassEntityManager& EntityManager, FMassExecutionContext& Context)
-{
-	Query.ForEachEntityChunk(EntityManager, Context, [](FMassExecutionContext& Context)
-		{
-			const FMassActorFragment* ActorIt = Context.GetFragmentView<FMassActorFragment>().GetData();
-			FTypedElementLabelColumn* LabelIt = Context.GetMutableFragmentView<FTypedElementLabelColumn>().GetData();
-			FTypedElementLabelHashColumn* LabelHashIt = Context.GetMutableFragmentView<FTypedElementLabelHashColumn>().GetData();
-			
-			const int32 NumEntities = Context.GetNumEntities();
-			for (int32 Index = 0; Index < NumEntities; ++Index)
+	DataStorage.RegisterQuery(
+		Select(
+			TEXT("Sync actor label to column"),
+			FProcessor(DSI::EQueryTickPhase::PrePhysics, DataStorage.GetQueryTickGroupName(DSI::EQueryTickGroups::SyncExternalToDataStorage))
+				.ForceToGameThread(true),
+			[](const FMassActorFragment& Actor, FTypedElementLabelColumn& Label, FTypedElementLabelHashColumn& LabelHash)
 			{
-				const FString& ActorLabel = ActorIt->Get()->GetActorLabel(false);
-				uint64 ActorLabelHash = CityHash64(
-					reinterpret_cast<const char*>(ActorLabel.GetCharArray().GetData()), 
-					ActorLabel.Len() * sizeof(ActorLabel.GetCharArray().GetData()[0]));
-				if (LabelHashIt->LabelHash != ActorLabelHash)
+				if (const AActor* ActorInstance = Actor.Get(); ActorInstance != nullptr)
 				{
-					LabelIt->Label = ActorLabel;
-					LabelHashIt->LabelHash = ActorLabelHash;
+					const FString& ActorLabel = ActorInstance->GetActorLabel(false);
+					uint64 ActorLabelHash = CityHash64(reinterpret_cast<const char*>(*ActorLabel), ActorLabel.Len() * sizeof(**ActorLabel));
+					if (LabelHash.LabelHash != ActorLabelHash)
+					{
+						Label.Label = ActorLabel;
+						LabelHash.LabelHash = ActorLabelHash;
+					}
 				}
-				
-				++ActorIt;
-				++LabelIt;
-				++LabelHashIt;
 			}
-		}
+		)
+		.Compile()
 	);
 }
 
-
-
-/**
- * UTypedElementLabelColumnToActorProcessor
- */
-
-UTypedElementLabelColumnToActorProcessor::UTypedElementLabelColumnToActorProcessor()
-	: Query(*this)
-{
-	ExecutionFlags = static_cast<int32>(EProcessorExecutionFlags::Editor);
-	ExecutionOrder.ExecuteInGroup = UE::Mass::ProcessorGroupNames::UpdateWorldFromMass;
-	ProcessingPhase = EMassProcessingPhase::FrameEnd;
-	bRequiresGameThreadExecution = true;
-}
-
-void UTypedElementLabelColumnToActorProcessor::ConfigureQueries()
-{
-	Query.AddRequirement(FMassActorFragment::StaticStruct(), EMassFragmentAccess::ReadWrite);
-	Query.AddRequirement(FTypedElementLabelColumn::StaticStruct(), EMassFragmentAccess::ReadOnly);
-	Query.AddRequirement(FTypedElementLabelHashColumn::StaticStruct(), EMassFragmentAccess::ReadOnly);
-	Query.AddTagRequirement(*FTypedElementSyncBackToWorldTag::StaticStruct(), EMassFragmentPresence::All);
-}
-
-void UTypedElementLabelColumnToActorProcessor::Execute(FMassEntityManager& EntityManager, FMassExecutionContext& Context)
+void UTypedElementActorLabelFactory::RegisterLabelColumnToActorQuery(ITypedElementDataStorageInterface& DataStorage) const
 {
-	Query.ForEachEntityChunk(EntityManager, Context, [](FMassExecutionContext& Context)
-		{
-			FMassActorFragment* ActorIt = Context.GetMutableFragmentView<FMassActorFragment>().GetData();
-			const FTypedElementLabelColumn* LabelIt = Context.GetFragmentView<FTypedElementLabelColumn>().GetData();
-			const FTypedElementLabelHashColumn* LabelHashIt = Context.GetFragmentView<FTypedElementLabelHashColumn>().GetData();
+	using namespace TypedElementQueryBuilder;
+	using DSI = ITypedElementDataStorageInterface;
 
-			const int32 NumEntities = Context.GetNumEntities();
-			for (int32 Index = 0; Index < NumEntities; ++Index)
+	DataStorage.RegisterQuery(
+		Select(
+			TEXT("Sync label column to actor"),
+			FProcessor(DSI::EQueryTickPhase::FrameEnd, DataStorage.GetQueryTickGroupName(DSI::EQueryTickGroups::SyncDataStorageToExternal))
+				.ForceToGameThread(true),
+			[](DSI::FQueryContext&, FMassActorFragment& Actor, const FTypedElementLabelColumn& Label, const FTypedElementLabelHashColumn& LabelHash)
 			{
-				const FString& ActorLabel = ActorIt->Get()->GetActorLabel(false);
-				uint64 ActorLabelHash = CityHash64(
-					reinterpret_cast<const char*>(ActorLabel.GetCharArray().GetData()),
-					ActorLabel.Len() * sizeof(ActorLabel.GetCharArray().GetData()[0]));
-				if (LabelHashIt->LabelHash != ActorLabelHash)
+				if (AActor* ActorInstance = Actor.GetMutable(); ActorInstance != nullptr)
 				{
-					ActorIt->GetMutable()->SetActorLabel(LabelIt->Label);
+					const FString& ActorLabel = ActorInstance->GetActorLabel(false);
+					uint64 ActorLabelHash = CityHash64(reinterpret_cast<const char*>(*ActorLabel), ActorLabel.Len() * sizeof(**ActorLabel));
+					if (LabelHash.LabelHash != ActorLabelHash)
+					{
+						ActorInstance->SetActorLabel(Label.Label);
+					}
 				}
-				
-				++ActorIt;
-				++LabelIt;
-				++LabelHashIt;
 			}
-		}
+		)
+		.Where()
+			.All<FTypedElementSyncBackToWorldTag>()
+		.Compile()
 	);
 }
diff --git a/Engine/Plugins/Experimental/TypedElementsDataStorage/Source/TypedElementsDataStorage/Private/Compatibility/TypedElementActorLabelProcessors.h b/Engine/Plugins/Experimental/TypedElementsDataStorage/Source/TypedElementsDataStorage/Private/Compatibility/TypedElementActorLabelProcessors.h
--- a/Engine/Plugins/Experimental/TypedElementsDataStorage/Source/TypedElementsDataStorage/Private/Compatibility/TypedElementActorLabelProcessors.h
+++ b/Engine/Plugins/Experimental/TypedElementsDataStorage/Source/TypedElementsDataStorage/Private/Compatibility/TypedElementActorLabelProcessors.h
@@ -2,48 +2,30 @@
 
 #pragma once
 
-#include "MassProcessor.h"
+#include "Elements/Interfaces/TypedElementDataStorageFactory.h"
 #include "UObject/ObjectMacros.h"
 
 #include "TypedElementActorLabelProcessors.generated.h"
 
-/**
- * Takes the label set on an actor and copies it to the Data Storage if they differ.
- */
-UCLASS()
-class UTypedElementActorLabelToColumnProcessor : public UMassProcessor
-{
-	GENERATED_BODY()
-
-public:
-	UTypedElementActorLabelToColumnProcessor();
-
-protected:
-	void ConfigureQueries() override;
-	void Execute(FMassEntityManager& EntityManager, FMassExecutionContext& Context) override;
-
-private:
-	FMassEntityQuery Query;
-};
-
-
 
-/**
- * Takes the label stored in the Data Storage and copies it to the actor's label if the FTypedElementSyncBackToWorldTag
- * has been set and the labels differ.
- */
 UCLASS()
-class UTypedElementLabelColumnToActorProcessor : public UMassProcessor
+class UTypedElementActorLabelFactory : public UTypedElementDataStorageFactory
 {
 	GENERATED_BODY()
 
 public:
-	UTypedElementLabelColumnToActorProcessor();
+	~UTypedElementActorLabelFactory() override = default;
 
-protected:
-	void ConfigureQueries() override;
-	void Execute(FMassEntityManager& EntityManager, FMassExecutionContext& Context) override;
+	void RegisterQueries(ITypedElementDataStorageInterface& DataStorage) const override;
 
 private:
-	FMassEntityQuery Query;
+	/**
+	 * Takes the label set on an actor and copies it to the Data Storage if they differ.
+	 */
+	void RegisterActorLabelToColumnQuery(ITypedElementDataStorageInterface& DataStorage) const;
+	/**
+	 * Takes the label stored in the Data Storage and copies it to the actor's label if the FTypedElementSyncBackToWorldTag
+	 * has been set and the labels differ.
+	 */
+	void RegisterLabelColumnToActorQuery(ITypedElementDataStorageInterface& DataStorage) const;
 };
\ No newline at end of file
diff --git a/Engine/Plugins/Experimental/TypedElementsDataStorage/Source/TypedElementsDataStorage/Private/Compatibility/TypedElementActorPackagePathToColumnProcessor.cpp b/Engine/Plugins/Experimental/TypedElementsDataStorage/Source/TypedElementsDataStorage/Private/Compatibility/TypedElementActorPackagePathToColumnProcessor.cpp
--- a/Engine/Plugins/Experimental/TypedElementsDataStorage/Source/TypedElementsDataStorage/Private/Compatibility/TypedElementActorPackagePathToColumnProcessor.cpp
+++ b/Engine/Plugins/Experimental/TypedElementsDataStorage/Source/TypedElementsDataStorage/Private/Compatibility/TypedElementActorPackagePathToColumnProcessor.cpp
@@ -3,48 +3,30 @@
 #include "Compatibility/TypedElementActorPackagePathToColumnProcessor.h"
 
 #include "Elements/Columns/TypedElementPackageColumns.h"
+#include "Elements/Framework/TypedElementQueryBuilder.h"
 #include "MassActorSubsystem.h"
-#include "MassExecutionContext.h"
-#include "TypedElementSubsystems.h"
+#include "UObject/Package.h"
 
-UTypedElementActorPackagePathToColumnProcessor::UTypedElementActorPackagePathToColumnProcessor()
-	: Query(*this)
+void UTypedElementActorPackagePathFactory::RegisterQueries(ITypedElementDataStorageInterface& DataStorage) const
 {
-	ExecutionFlags = static_cast<int32>(EProcessorExecutionFlags::Editor);
-	ObservedType = FMassActorFragment::StaticStruct();
-	Operation = EMassObservedOperation::Add;
-	bRequiresGameThreadExecution = true;
-}
-
-void UTypedElementActorPackagePathToColumnProcessor::ConfigureQueries()
-{
-	Query.AddRequirement(FMassActorFragment::StaticStruct(), EMassFragmentAccess::ReadOnly);
-	Query.AddRequirement(FTypedElementPackagePathColumn::StaticStruct(), EMassFragmentAccess::ReadWrite);
-	Query.AddRequirement(FTypedElementPackageLoadedPathColumn::StaticStruct(), EMassFragmentAccess::ReadWrite);
-	
-	ProcessorRequirements.AddSubsystemRequirement(UTypedElementDataStorageSubsystem::StaticClass(), EMassFragmentAccess::ReadWrite);
-}
+	using namespace TypedElementQueryBuilder;
+	using DSI = ITypedElementDataStorageInterface;
 
-void UTypedElementActorPackagePathToColumnProcessor::Execute(FMassEntityManager& EntityManager, FMassExecutionContext& Context)
-{
-	Query.ForEachEntityChunk(EntityManager, Context, [](FMassExecutionContext& Context)
-		{
-			const FMassActorFragment* ActorIt = Context.GetFragmentView<FMassActorFragment>().GetData();
-			FTypedElementPackagePathColumn* PathIt = Context.GetMutableFragmentView<FTypedElementPackagePathColumn>().GetData();
-			FTypedElementPackageLoadedPathColumn* LoadedPathIt = Context.GetMutableFragmentView<FTypedElementPackageLoadedPathColumn>().GetData();
-			
-			const int32 EntityCount = Context.GetNumEntities();
-			for (int32 Counter=0; Counter < EntityCount; ++Counter)
+	DataStorage.RegisterQuery(
+		Select(
+			TEXT("Sync actor package info to columns"),
+			FObserver(FObserver::EEvent::Add, FMassActorFragment::StaticStruct())
+				.ForceToGameThread(true),
+			[](const FMassActorFragment& Actor, FTypedElementPackagePathColumn& Path, FTypedElementPackageLoadedPathColumn& LoadedPath)
 			{
-				UPackage* Target = ActorIt->Get()->GetPackage();
-
-				Target->GetPathName(nullptr, PathIt->Path);
-				LoadedPathIt->LoadedPath = Target->GetLoadedPath();
-				
-				++ActorIt;
-				++PathIt;
-				++LoadedPathIt;
+				if (const AActor* ActorInstance = Actor.Get(); ActorInstance != nullptr)
+				{
+					const UPackage* Target = ActorInstance->GetPackage();
+					Target->GetPathName(nullptr, Path.Path);
+					LoadedPath.LoadedPath = Target->GetLoadedPath();
+				}
 			}
-		}
+		)
+		.Compile()
 	);
 }
diff --git a/Engine/Plugins/Experimental/TypedElementsDataStorage/Source/TypedElementsDataStorage/Private/Compatibility/TypedElementActorPackagePathToColumnProcessor.h b/Engine/Plugins/Experimental/TypedElementsDataStorage/Source/TypedElementsDataStorage/Private/Compatibility/TypedElementActorPackagePathToColumnProcessor.h
--- a/Engine/Plugins/Experimental/TypedElementsDataStorage/Source/TypedElementsDataStorage/Private/Compatibility/TypedElementActorPackagePathToColumnProcessor.h
+++ b/Engine/Plugins/Experimental/TypedElementsDataStorage/Source/TypedElementsDataStorage/Private/Compatibility/TypedElementActorPackagePathToColumnProcessor.h
@@ -2,23 +2,18 @@
 
 #pragma once
 
-#include "MassObserverProcessor.h"
+#include "Elements/Interfaces/TypedElementDataStorageFactory.h"
 #include "UObject/ObjectMacros.h"
 
 #include "TypedElementActorPackagePathToColumnProcessor.generated.h"
 
 UCLASS()
-class UTypedElementActorPackagePathToColumnProcessor : public UMassObserverProcessor
+class UTypedElementActorPackagePathFactory : public UTypedElementDataStorageFactory
 {
 	GENERATED_BODY()
 
 public:
-	UTypedElementActorPackagePathToColumnProcessor();
+	~UTypedElementActorPackagePathFactory() override = default;
 
-protected:
-	void ConfigureQueries() override;
-	void Execute(FMassEntityManager& EntityManager, FMassExecutionContext& Context) override;
-
-private:
-	FMassEntityQuery Query;
+	void RegisterQueries(ITypedElementDataStorageInterface& DataStorage) const override;
 };
diff --git a/Engine/Plugins/Experimental/TypedElementsDataStorage/Source/TypedElementsDataStorage/Private/Compatibility/TypedElementActorTransformProcessors.cpp b/Engine/Plugins/Experimental/TypedElementsDataStorage/Source/TypedElementsDataStorage/Private/Compatibility/TypedElementActorTransformProcessors.cpp
--- a/Engine/Plugins/Experimental/TypedElementsDataStorage/Source/TypedElementsDataStorage/Private/Compatibility/TypedElementActorTransformProcessors.cpp
+++ b/Engine/Plugins/Experimental/TypedElementsDataStorage/Source/TypedElementsDataStorage/Private/Compatibility/TypedElementActorTransformProcessors.cpp
@@ -2,132 +2,93 @@
 
 #include "Compatibility/TypedElementActorTransformProcessors.h"
 
+
 #include "Elements/Columns/TypedElementMiscColumns.h"
 #include "Elements/Columns/TypedElementTransformColumns.h"
+#include "Elements/Framework/TypedElementQueryBuilder.h"
 #include "MassActorSubsystem.h"
-#include "MassExecutionContext.h"
-
-/**
- * UTypedElementActorAddTransformColumnProcessor
- */
-
-UTypedElementActorAddTransformColumnProcessor::UTypedElementActorAddTransformColumnProcessor()
-	: Query(*this)
-{
-	ExecutionFlags = static_cast<int32>(EProcessorExecutionFlags::Editor);
-	// Run this before the world gets synced to MASS so a new transform is initialized in the same frame.
-	ExecutionOrder.ExecuteBefore.Add(UE::Mass::ProcessorGroupNames::SyncWorldToMass);
-	ProcessingPhase = EMassProcessingPhase::PrePhysics;
-	bRequiresGameThreadExecution = true;
-}
 
-void UTypedElementActorAddTransformColumnProcessor::ConfigureQueries()
+void UTypedElementActorTransformFactory::RegisterQueries(ITypedElementDataStorageInterface& DataStorage) const
 {
-	Query.AddRequirement(FMassActorFragment::StaticStruct(), EMassFragmentAccess::ReadOnly);
-	Query.AddRequirement(FTypedElementLocalTransformColumn::StaticStruct(), EMassFragmentAccess::ReadOnly, EMassFragmentPresence::None);
+	RegisterActorAddTransformColumn(DataStorage);
+	RegisterActorLocalTransformToColumn(DataStorage);
+	RegisterLocalTransformColumnToActor(DataStorage);
 }
 
-void UTypedElementActorAddTransformColumnProcessor::Execute(FMassEntityManager& EntityManager, FMassExecutionContext& Context)
+void UTypedElementActorTransformFactory::RegisterActorAddTransformColumn(ITypedElementDataStorageInterface& DataStorage) const
 {
-	Query.ForEachEntityChunk(EntityManager, Context, [](FMassExecutionContext& Context)
-		{
-			const FMassActorFragment* ActorIt = Context.GetFragmentView<FMassActorFragment>().GetData();
-			
-			for (FMassEntityHandle Entity : Context.GetEntities())
+	using namespace TypedElementQueryBuilder;
+	using DSI = ITypedElementDataStorageInterface;
+
+	DataStorage.RegisterQuery(
+		Select(
+			TEXT("Add transform column to actor"),
+			FProcessor(DSI::EQueryTickPhase::PrePhysics, 
+				DataStorage.GetQueryTickGroupName(DSI::EQueryTickGroups::PrepareSyncExternalToDataStorage))
+				.SetBeforeGroup(DataStorage.GetQueryTickGroupName(DSI::EQueryTickGroups::SyncExternalToDataStorage))
+				.ForceToGameThread(true),
+			[](DSI::FQueryContext& Context, TypedElementRowHandle Row, const FMassActorFragment& Actor)
 			{
-				if (ActorIt->Get()->GetRootComponent())
+				const AActor* ActorInstance = Actor.Get();
+				if (ActorInstance != nullptr && ActorInstance->GetRootComponent())
 				{
-					Context.Defer().AddFragment_RuntimeCheck<FTypedElementLocalTransformColumn>(Entity);
+					Context.AddColumns<FTypedElementLocalTransformColumn>(Row);
 				}
-				
-				++ActorIt;
 			}
-		}
+		)
+		.Where()
+			.None<FTypedElementLocalTransformColumn>()
+		.Compile()
 	);
 }
 
-
-/**
- * UTypedElementActorLocalTransformToColumnProcessor
- */
-
-UTypedElementActorLocalTransformToColumnProcessor::UTypedElementActorLocalTransformToColumnProcessor()
-	: Query(*this)
-{
-	ExecutionFlags = static_cast<int32>(EProcessorExecutionFlags::Editor);
-	ExecutionOrder.ExecuteInGroup = UE::Mass::ProcessorGroupNames::SyncWorldToMass;
-	ProcessingPhase = EMassProcessingPhase::PrePhysics;
-	bRequiresGameThreadExecution = true;
-}
-
-void UTypedElementActorLocalTransformToColumnProcessor::ConfigureQueries()
-{
-	Query.AddRequirement(FMassActorFragment::StaticStruct(), EMassFragmentAccess::ReadOnly);
-	Query.AddRequirement(FTypedElementLocalTransformColumn::StaticStruct(), EMassFragmentAccess::ReadWrite);
-}
-
-void UTypedElementActorLocalTransformToColumnProcessor::Execute(FMassEntityManager& EntityManager, FMassExecutionContext& Context)
+void UTypedElementActorTransformFactory::RegisterActorLocalTransformToColumn(ITypedElementDataStorageInterface& DataStorage) const
 {
-	Query.ForEachEntityChunk(EntityManager, Context, [](FMassExecutionContext& Context)
-		{
-			const FMassActorFragment* ActorIt = Context.GetFragmentView<FMassActorFragment>().GetData();
-			FTypedElementLocalTransformColumn* TransformIt = Context.GetMutableFragmentView<FTypedElementLocalTransformColumn>().GetData();
-
-			for (FMassEntityHandle Entity : Context.GetEntities())
+	using namespace TypedElementQueryBuilder;
+	using DSI = ITypedElementDataStorageInterface;
+
+	DataStorage.RegisterQuery(
+		Select(
+			TEXT("Sync actor transform to column"),
+			FProcessor(DSI::EQueryTickPhase::PrePhysics, DataStorage.GetQueryTickGroupName(DSI::EQueryTickGroups::SyncExternalToDataStorage))
+				.ForceToGameThread(true),
+			[](DSI::FQueryContext& Context, TypedElementRowHandle Row, const FMassActorFragment& Actor, FTypedElementLocalTransformColumn& Transform)
 			{
-				const AActor* ActorInstance = ActorIt->Get();
+				const AActor* ActorInstance = Actor.Get();
 				if (ActorInstance != nullptr && ActorInstance->GetRootComponent() != nullptr)
 				{
-					TransformIt->Transform = ActorIt->Get()->GetActorTransform();
+					Transform.Transform = ActorInstance->GetActorTransform();
 				}
 				else
 				{
-					Context.Defer().RemoveFragment_RuntimeCheck<FTypedElementLocalTransformColumn>(Entity);
+					Context.RemoveColumns<FTypedElementLocalTransformColumn>(Row);
 				}
-				
-				++ActorIt;
-				++TransformIt;
 			}
-		}
+		)
+		.Compile()
 	);
 }
 
-
-/**
- * UTypedElementTransformColumnToActorProcessor
- */
-
-UTypedElementTransformColumnToActorProcessor::UTypedElementTransformColumnToActorProcessor()
-	: Query(*this)
-{
-	ExecutionFlags = static_cast<int32>(EProcessorExecutionFlags::Editor);
-	ExecutionOrder.ExecuteInGroup = UE::Mass::ProcessorGroupNames::UpdateWorldFromMass;
-	ProcessingPhase = EMassProcessingPhase::FrameEnd;
-	bRequiresGameThreadExecution = true;
-}
-
-void UTypedElementTransformColumnToActorProcessor::ConfigureQueries()
-{
-	Query.AddRequirement(FMassActorFragment::StaticStruct(), EMassFragmentAccess::ReadWrite);
-	Query.AddRequirement(FTypedElementLocalTransformColumn::StaticStruct(), EMassFragmentAccess::ReadOnly);
-	Query.AddTagRequirement(*FTypedElementSyncBackToWorldTag::StaticStruct(), EMassFragmentPresence::All);
-}
-
-void UTypedElementTransformColumnToActorProcessor::Execute(FMassEntityManager& EntityManager, FMassExecutionContext& Context)
+void UTypedElementActorTransformFactory::RegisterLocalTransformColumnToActor(ITypedElementDataStorageInterface& DataStorage) const
 {
-	Query.ForEachEntityChunk(EntityManager, Context, [](FMassExecutionContext& Context)
-		{
-			FMassActorFragment* ActorIt = Context.GetMutableFragmentView<FMassActorFragment>().GetData();
-			const FTypedElementLocalTransformColumn* TransformIt = Context.GetFragmentView<FTypedElementLocalTransformColumn>().GetData();
-			
-			const int32 NumEntities = Context.GetNumEntities();
-			for (int32 Index = 0; Index < NumEntities; ++Index)
+	using namespace TypedElementQueryBuilder;
+	using DSI = ITypedElementDataStorageInterface;
+
+	DataStorage.RegisterQuery(
+		Select(
+			TEXT("Sync transform column to actor"),
+			FProcessor(DSI::EQueryTickPhase::FrameEnd, DataStorage.GetQueryTickGroupName(DSI::EQueryTickGroups::SyncDataStorageToExternal))
+				.ForceToGameThread(true),
+			[](FMassActorFragment& Actor, const FTypedElementLocalTransformColumn& Transform)
 			{
-				ActorIt->GetMutable()->SetActorTransform(TransformIt->Transform);
-				
-				++ActorIt;
-				++TransformIt;
+				if (AActor* ActorInstance = Actor.GetMutable(); ActorInstance != nullptr)
+				{
+					ActorInstance->SetActorTransform(Transform.Transform);
+				}
 			}
-		}
+		)
+		.Where()
+			.All< FTypedElementSyncBackToWorldTag>()
+		.Compile()
 	);
 }
diff --git a/Engine/Plugins/Experimental/TypedElementsDataStorage/Source/TypedElementsDataStorage/Private/Compatibility/TypedElementActorTransformProcessors.h b/Engine/Plugins/Experimental/TypedElementsDataStorage/Source/TypedElementsDataStorage/Private/Compatibility/TypedElementActorTransformProcessors.h
--- a/Engine/Plugins/Experimental/TypedElementsDataStorage/Source/TypedElementsDataStorage/Private/Compatibility/TypedElementActorTransformProcessors.h
+++ b/Engine/Plugins/Experimental/TypedElementsDataStorage/Source/TypedElementsDataStorage/Private/Compatibility/TypedElementActorTransformProcessors.h
@@ -2,70 +2,35 @@
 
 #pragma once
 
-#include "MassProcessor.h"
+#include "Elements/Interfaces/TypedElementDataStorageFactory.h"
 #include "UObject/ObjectMacros.h"
 
 #include "TypedElementActorTransformProcessors.generated.h"
 
-/**
- * Checks actors that don't have a tranform column and adds one if an actor has been
- * assigned a transform.
- */
 UCLASS()
-class UTypedElementActorAddTransformColumnProcessor : public UMassProcessor
+class UTypedElementActorTransformFactory : public UTypedElementDataStorageFactory
 {
 	GENERATED_BODY()
 
 public:
-	UTypedElementActorAddTransformColumnProcessor();
+	~UTypedElementActorTransformFactory() override = default;
 
-protected:
-	void ConfigureQueries() override;
-	void Execute(FMassEntityManager& EntityManager, FMassExecutionContext& Context) override;
+	void RegisterQueries(ITypedElementDataStorageInterface& DataStorage) const override;
 
 private:
-	FMassEntityQuery Query;
+	/**
+	 * Checks actors that don't have a tranform column and adds one if an actor has been
+	 * assigned a transform.
+	 */
+	void RegisterActorAddTransformColumn(ITypedElementDataStorageInterface& DataStorage) const;
+	/**
+	 * Takes the transform set on an actor and copies it to the Data Storage or removes the
+	 * transform column if there's not transform available anymore.
+	 */
+	void RegisterActorLocalTransformToColumn(ITypedElementDataStorageInterface& DataStorage) const;
+	/**
+	 * Takes the transform stored in the Data Storage and copies it to the actor's tranform if
+	 * the FTypedElementSyncBackToWorldTag has been set.
+	 */
+	void RegisterLocalTransformColumnToActor(ITypedElementDataStorageInterface& DataStorage) const;
 };
-
-
-/**
- * Takes the transform set on an actor and copies it to the Data Storage or removes the 
- * transform column if there's not transform available anymore.
- */
-UCLASS()
-class UTypedElementActorLocalTransformToColumnProcessor : public UMassProcessor
-{
-	GENERATED_BODY()
-
-public:
-	UTypedElementActorLocalTransformToColumnProcessor();
-
-protected:
-	void ConfigureQueries() override;
-	void Execute(FMassEntityManager& EntityManager, FMassExecutionContext& Context) override;
-
-private:
-	FMassEntityQuery Query;
-};
-
-
-
-/**
- * Takes the transform stored in the Data Storage and copies it to the actor's tranform if 
- * the FTypedElementSyncBackToWorldTag has been set.
- */
-UCLASS()
-class UTypedElementTransformColumnToActorProcessor : public UMassProcessor
-{
-	GENERATED_BODY()
-
-public:
-	UTypedElementTransformColumnToActorProcessor();
-
-protected:
-	void ConfigureQueries() override;
-	void Execute(FMassEntityManager& EntityManager, FMassExecutionContext& Context) override;
-
-private:
-	FMassEntityQuery Query;
-};
\ No newline at end of file
diff --git a/Engine/Plugins/Experimental/TypedElementsDataStorage/Source/TypedElementsDataStorage/Private/Processors/TypedElementMiscProcessors.cpp b/Engine/Plugins/Experimental/TypedElementsDataStorage/Source/TypedElementsDataStorage/Private/Processors/TypedElementMiscProcessors.cpp
--- a/Engine/Plugins/Experimental/TypedElementsDataStorage/Source/TypedElementsDataStorage/Private/Processors/TypedElementMiscProcessors.cpp
+++ b/Engine/Plugins/Experimental/TypedElementsDataStorage/Source/TypedElementsDataStorage/Private/Processors/TypedElementMiscProcessors.cpp
@@ -2,35 +2,27 @@
 
 #include "Processors/TypedElementMiscProcessors.h"
 
-#include "MassCommonTypes.h"
-#include "MassExecutionContext.h"
 #include "Elements/Columns/TypedElementMiscColumns.h"
+#include "Elements/Framework/TypedElementQueryBuilder.h"
 
-/**
- * UTypedElementRemoveSyncToWorldTagProcessor
- */
-
-UTypedElementRemoveSyncToWorldTagProcessor::UTypedElementRemoveSyncToWorldTagProcessor()
-	: Query(*this)
-{
-	ExecutionFlags = static_cast<int32>(EProcessorExecutionFlags::Editor);
-	ExecutionOrder.ExecuteAfter.Add(UE::Mass::ProcessorGroupNames::UpdateWorldFromMass);
-	ProcessingPhase = EMassProcessingPhase::FrameEnd;
-}
-
-void UTypedElementRemoveSyncToWorldTagProcessor::ConfigureQueries()
+void UTypedElementRemoveSyncToWorldTagFactory::RegisterQueries(ITypedElementDataStorageInterface& DataStorage) const
 {
-	Query.AddTagRequirement(*FTypedElementSyncBackToWorldTag::StaticStruct(), EMassFragmentPresence::All);
-}
+	using namespace TypedElementQueryBuilder;
+	using DSI = ITypedElementDataStorageInterface;
 
-void UTypedElementRemoveSyncToWorldTagProcessor::Execute(FMassEntityManager& EntityManager, FMassExecutionContext& Context)
-{
-	Query.ForEachEntityChunk(EntityManager, Context, [](FMassExecutionContext& Context)
-		{
-			for (FMassEntityHandle Entity : Context.GetEntities())
+	DataStorage.RegisterQuery(
+		Select(
+			TEXT("Remove 'sync to world' tag"),
+			FProcessor(DSI::EQueryTickPhase::FrameEnd, 
+				DataStorage.GetQueryTickGroupName(DSI::EQueryTickGroups::FinalizeSyncDataStorageToExternal))
+				.SetAfterGroup(DataStorage.GetQueryTickGroupName(DSI::EQueryTickGroups::SyncDataStorageToExternal)),
+			[](DSI::FQueryContext& Context, const TypedElementRowHandle* Rows)
 			{
-				Context.Defer().RemoveTag_RuntimeCheck<FTypedElementSyncBackToWorldTag>(Entity);
+				Context.RemoveColumns<FTypedElementSyncBackToWorldTag>(TConstArrayView<TypedElementRowHandle>(Rows, Context.GetRowCount()));
 			}
-		}
+		)
+		.Where()
+			.All<FTypedElementSyncBackToWorldTag>()
+		.Compile()
 	);
 }
diff --git a/Engine/Plugins/Experimental/TypedElementsDataStorage/Source/TypedElementsDataStorage/Private/Processors/TypedElementMiscProcessors.h b/Engine/Plugins/Experimental/TypedElementsDataStorage/Source/TypedElementsDataStorage/Private/Processors/TypedElementMiscProcessors.h
--- a/Engine/Plugins/Experimental/TypedElementsDataStorage/Source/TypedElementsDataStorage/Private/Processors/TypedElementMiscProcessors.h
+++ b/Engine/Plugins/Experimental/TypedElementsDataStorage/Source/TypedElementsDataStorage/Private/Processors/TypedElementMiscProcessors.h
@@ -2,7 +2,7 @@
 
 #pragma once
 
-#include "MassProcessor.h"
+#include "Elements/Interfaces/TypedElementDataStorageFactory.h"
 #include "UObject/ObjectMacros.h"
 
 #include "TypedElementMiscProcessors.generated.h"
@@ -11,17 +11,12 @@
  * Removes all FTypedElementSyncBackToWorldTags at the end of an update cycle.
  */
 UCLASS()
-class UTypedElementRemoveSyncToWorldTagProcessor : public UMassProcessor
+class UTypedElementRemoveSyncToWorldTagFactory : public UTypedElementDataStorageFactory
 {
 	GENERATED_BODY()
 
 public:
-	UTypedElementRemoveSyncToWorldTagProcessor();
+	~UTypedElementRemoveSyncToWorldTagFactory() override = default;
 
-protected:
-	void ConfigureQueries() override;
-	void Execute(FMassEntityManager& EntityManager, FMassExecutionContext& Context) override;
-
-private:
-	FMassEntityQuery Query;
-};
\ No newline at end of file
+	void RegisterQueries(ITypedElementDataStorageInterface& DataStorage) const override;
+};
diff --git a/Engine/Plugins/Experimental/TypedElementsDataStorage/Source/TypedElementsDataStorage/Private/UI/TypedElementSlateWidgetReferenceColumnUpdateProcessor.cpp b/Engine/Plugins/Experimental/TypedElementsDataStorage/Source/TypedElementsDataStorage/Private/UI/TypedElementSlateWidgetReferenceColumnUpdateProcessor.cpp
--- a/Engine/Plugins/Experimental/TypedElementsDataStorage/Source/TypedElementsDataStorage/Private/UI/TypedElementSlateWidgetReferenceColumnUpdateProcessor.cpp
+++ b/Engine/Plugins/Experimental/TypedElementsDataStorage/Source/TypedElementsDataStorage/Private/UI/TypedElementSlateWidgetReferenceColumnUpdateProcessor.cpp
@@ -3,60 +3,60 @@
 #include "UI/TypedElementSlateWidgetReferenceColumnUpdateProcessor.h"
 
 #include "Elements/Columns/TypedElementSlateWidgetColumns.h"
-#include "MassCommonTypes.h"
-#include "MassExecutionContext.h"
+#include "Elements/Framework/TypedElementQueryBuilder.h"
 
-UTypedElementSlateWidgetReferenceColumnUpdateProcessor::UTypedElementSlateWidgetReferenceColumnUpdateProcessor()
-	: ColumnRemovalQuery(*this)
-	, RowDeletionQuery(*this)
+void UTypedElementSlateWidgetReferenceColumnUpdateFactory::RegisterQueries(ITypedElementDataStorageInterface& DataStorage) const
 {
-	ExecutionFlags = static_cast<int32>(EProcessorExecutionFlags::Editor);
-	ExecutionOrder.ExecuteInGroup = UE::Mass::ProcessorGroupNames::SyncWorldToMass;
-	ProcessingPhase = EMassProcessingPhase::PrePhysics;
+	RegisterDeleteRowOnWidgetDeleteQuery(DataStorage);
+	RegisterDeleteColumnOnWidgetDeleteQuery(DataStorage);
 }
 
-void UTypedElementSlateWidgetReferenceColumnUpdateProcessor::ConfigureQueries()
+void UTypedElementSlateWidgetReferenceColumnUpdateFactory::RegisterDeleteRowOnWidgetDeleteQuery(ITypedElementDataStorageInterface& DataStorage) const
 {
-	ColumnRemovalQuery.AddRequirement(FTypedElementSlateWidgetReferenceColumn::StaticStruct(), EMassFragmentAccess::ReadOnly);
-	ColumnRemovalQuery.AddTagRequirement(*FTypedElementSlateWidgetReferenceDeletesRowTag::StaticStruct(), EMassFragmentPresence::None);
-	
-	RowDeletionQuery.AddRequirement(FTypedElementSlateWidgetReferenceColumn::StaticStruct(), EMassFragmentAccess::ReadOnly);
-	RowDeletionQuery.AddTagRequirement(*FTypedElementSlateWidgetReferenceDeletesRowTag::StaticStruct(), EMassFragmentPresence::All);
-}
-
-void UTypedElementSlateWidgetReferenceColumnUpdateProcessor::Execute(
-	FMassEntityManager& EntityManager, FMassExecutionContext& Context)
-{
-	ColumnRemovalQuery.ForEachEntityChunk(EntityManager, Context, [](FMassExecutionContext& Context)
-		{
-			CheckStatus<false>(Context);
-		});
+	using namespace TypedElementQueryBuilder;
+	using DSI = ITypedElementDataStorageInterface;
 
-	RowDeletionQuery.ForEachEntityChunk(EntityManager, Context, [](FMassExecutionContext& Context)
-		{
-			CheckStatus<true>(Context);
-		});
+	DataStorage.RegisterQuery(
+		Select(
+			TEXT("Delete row with deleted widget"),
+			FProcessor(DSI::EQueryTickPhase::PrePhysics, 
+				DataStorage.GetQueryTickGroupName(DSI::EQueryTickGroups::PrepareSyncExternalToDataStorage))
+				.SetBeforeGroup(DataStorage.GetQueryTickGroupName(DSI::EQueryTickGroups::SyncExternalToDataStorage)),
+			[](DSI::FQueryContext& Context, TypedElementRowHandle Row, const FTypedElementSlateWidgetReferenceColumn& WidgetReference)
+			{
+				if (!WidgetReference.Widget.IsValid())
+				{
+					Context.RemoveRow(Row);
+				}
+			}
+		)
+		.Where()
+			.All<FTypedElementSlateWidgetReferenceDeletesRowTag>()
+		.Compile()
+	);
 }
 
-template<bool DeleteRow>
-void UTypedElementSlateWidgetReferenceColumnUpdateProcessor::CheckStatus(FMassExecutionContext& Context)
+void UTypedElementSlateWidgetReferenceColumnUpdateFactory::RegisterDeleteColumnOnWidgetDeleteQuery(ITypedElementDataStorageInterface& DataStorage) const
 {
-	int32 Index = 0;
-	TArrayView<FTypedElementSlateWidgetReferenceColumn> WidgetColumns =
-		Context.GetMutableFragmentView<FTypedElementSlateWidgetReferenceColumn>();
-	for (FTypedElementSlateWidgetReferenceColumn& WidgetColumn : WidgetColumns)
-	{
-		if (!WidgetColumn.Widget.IsValid())
-		{
-			if constexpr (DeleteRow)
-			{
-				Context.Defer().DestroyEntity(Context.GetEntity(Index));
-			}
-			else
+	using namespace TypedElementQueryBuilder;
+	using DSI = ITypedElementDataStorageInterface;
+
+	DataStorage.RegisterQuery(
+		Select(
+			TEXT("Delete widget column for deleted widget"),
+			FProcessor(DSI::EQueryTickPhase::PrePhysics,
+				DataStorage.GetQueryTickGroupName(DSI::EQueryTickGroups::PrepareSyncExternalToDataStorage))
+				.SetBeforeGroup(DataStorage.GetQueryTickGroupName(DSI::EQueryTickGroups::SyncExternalToDataStorage)),
+			[](DSI::FQueryContext& Context, TypedElementRowHandle Row, const FTypedElementSlateWidgetReferenceColumn& WidgetReference)
 			{
-				Context.Defer().RemoveFragment_RuntimeCheck<FTypedElementSlateWidgetReferenceColumn>(Context.GetEntity(Index));
+				if (!WidgetReference.Widget.IsValid())
+				{
+					Context.RemoveColumns<FTypedElementSlateWidgetReferenceColumn>(Row);
+				}
 			}
-		}
-		++Index;
-	}
+		)
+		.Where()
+			.All<FTypedElementSlateWidgetReferenceDeletesRowTag>()
+		.Compile()
+	);
 }
diff --git a/Engine/Plugins/Experimental/TypedElementsDataStorage/Source/TypedElementsDataStorage/Private/UI/TypedElementSlateWidgetReferenceColumnUpdateProcessor.h b/Engine/Plugins/Experimental/TypedElementsDataStorage/Source/TypedElementsDataStorage/Private/UI/TypedElementSlateWidgetReferenceColumnUpdateProcessor.h
--- a/Engine/Plugins/Experimental/TypedElementsDataStorage/Source/TypedElementsDataStorage/Private/UI/TypedElementSlateWidgetReferenceColumnUpdateProcessor.h
+++ b/Engine/Plugins/Experimental/TypedElementsDataStorage/Source/TypedElementsDataStorage/Private/UI/TypedElementSlateWidgetReferenceColumnUpdateProcessor.h
@@ -2,32 +2,27 @@
 
 #pragma once
 
-#include "MassObserverProcessor.h"
-#include "MassProcessor.h"
+#include "Elements/Interfaces/TypedElementDataStorageFactory.h"
+#include "UObject/ObjectMacros.h"
 
 #include "TypedElementSlateWidgetReferenceColumnUpdateProcessor.generated.h"
 
 /**
- * Processor that checks whether or not a widget still exists. If it has been deleted
+ * Queries that check whether or not a widget still exists. If it has been deleted
  * then it will remove the column from the Data Storage or deletes the entire row if
  * the FTypedElementSlateWidgetReferenceDeletesRowTag was found.
  */
 UCLASS()
-class UTypedElementSlateWidgetReferenceColumnUpdateProcessor : public UMassProcessor
+class UTypedElementSlateWidgetReferenceColumnUpdateFactory : public UTypedElementDataStorageFactory
 {
 	GENERATED_BODY()
 
 public:
-	UTypedElementSlateWidgetReferenceColumnUpdateProcessor();
+	~UTypedElementSlateWidgetReferenceColumnUpdateFactory() override = default;
 
-protected:
-	void ConfigureQueries() override;
-	void Execute(FMassEntityManager& EntityManager, FMassExecutionContext& Context) override;
+	void RegisterQueries(ITypedElementDataStorageInterface& DataStorage) const override;
 
 private:
-	template<bool DeleteRow>
-	static void CheckStatus(FMassExecutionContext& Context);
-
-	FMassEntityQuery ColumnRemovalQuery;
-	FMassEntityQuery RowDeletionQuery;
-};
\ No newline at end of file
+	void RegisterDeleteRowOnWidgetDeleteQuery(ITypedElementDataStorageInterface& DataStorage) const;
+	void RegisterDeleteColumnOnWidgetDeleteQuery(ITypedElementDataStorageInterface& DataStorage) const;
+};
