Added the ability to retrieve the row from TEDS that contains the provided actor.

diff --git a/Engine/Plugins/Experimental/TypedElementsDataStorage/Source/TypedElementsDataStorage/Private/TypedElementDatabaseCompatibility.cpp b/Engine/Plugins/Experimental/TypedElementsDataStorage/Source/TypedElementsDataStorage/Private/TypedElementDatabaseCompatibility.cpp
--- a/Engine/Plugins/Experimental/TypedElementsDataStorage/Source/TypedElementsDataStorage/Private/TypedElementDatabaseCompatibility.cpp
+++ b/Engine/Plugins/Experimental/TypedElementsDataStorage/Source/TypedElementsDataStorage/Private/TypedElementDatabaseCompatibility.cpp
@@ -58,6 +58,19 @@ void UTypedElementDatabaseCompatibility::RemoveCompatibleObject(AActor* Actor)
 	}
 }
 
+TypedElementRowHandle UTypedElementDatabaseCompatibility::FindRowWithCompatibleObject(AActor* Actor) const
+{
+	if (Storage && ActorSubsystem && Storage->IsAvailable())
+	{
+		FMassEntityHandle Entity = ActorSubsystem->GetEntityHandleFromActor(Actor);
+		return Entity.IsValid() ? Entity.AsNumber() : TypedElementInvalidRowHandle;
+	}
+	else
+	{
+		return TypedElementInvalidRowHandle;
+	}
+}
+
 void UTypedElementDatabaseCompatibility::Prepare()
 {
 	UMassActorEditorSubsystem* MassActorEditorSubsystem = Storage->GetExternalSystem<UMassActorEditorSubsystem>();
diff --git a/Engine/Plugins/Experimental/TypedElementsDataStorage/Source/TypedElementsDataStorage/Private/TypedElementDatabaseCompatibility.h b/Engine/Plugins/Experimental/TypedElementsDataStorage/Source/TypedElementsDataStorage/Private/TypedElementDatabaseCompatibility.h
--- a/Engine/Plugins/Experimental/TypedElementsDataStorage/Source/TypedElementsDataStorage/Private/TypedElementDatabaseCompatibility.h
+++ b/Engine/Plugins/Experimental/TypedElementsDataStorage/Source/TypedElementsDataStorage/Private/TypedElementDatabaseCompatibility.h
@@ -28,6 +28,7 @@ public:
 
 	void AddCompatibleObject(AActor* Actor) override;
 	void RemoveCompatibleObject(AActor* Actor) override;
+	TypedElementRowHandle FindRowWithCompatibleObject(AActor* Actor) const;
 
 private:
 	void Prepare();
diff --git a/Engine/Source/Runtime/TypedElementFramework/Public/Elements/Interfaces/TypedElementDataStorageCompatibilityInterface.h b/Engine/Source/Runtime/TypedElementFramework/Public/Elements/Interfaces/TypedElementDataStorageCompatibilityInterface.h
--- a/Engine/Source/Runtime/TypedElementFramework/Public/Elements/Interfaces/TypedElementDataStorageCompatibilityInterface.h
+++ b/Engine/Source/Runtime/TypedElementFramework/Public/Elements/Interfaces/TypedElementDataStorageCompatibilityInterface.h
@@ -2,6 +2,7 @@
 
 #pragma once
 
+#include "Elements/Interfaces/TypedElementDataStorageInterface.h"
 #include "UObject/Interface.h"
 #include "UObject/ObjectMacros.h"
 
@@ -26,4 +27,5 @@ class TYPEDELEMENTFRAMEWORK_API ITypedElementDataStorageCompatibilityInterface
 public:
 	virtual void AddCompatibleObject(AActor* Actor) = 0;
 	virtual void RemoveCompatibleObject(AActor* Actor) = 0;
+	virtual TypedElementRowHandle FindRowWithCompatibleObject(AActor* Actor) const = 0;
 };
\ No newline at end of file
