Removed unintial pragma to disable optmizations.

diff --git a/Engine/Plugins/Experimental/TypedElementsDataStorage/Source/TypedElementsDataStorage/Private/UI/TypedElementSlateWidgetReferenceColumnUpdateProcessor.cpp b/Engine/Plugins/Experimental/TypedElementsDataStorage/Source/TypedElementsDataStorage/Private/UI/TypedElementSlateWidgetReferenceColumnUpdateProcessor.cpp
--- a/Engine/Plugins/Experimental/TypedElementsDataStorage/Source/TypedElementsDataStorage/Private/UI/TypedElementSlateWidgetReferenceColumnUpdateProcessor.cpp
+++ b/Engine/Plugins/Experimental/TypedElementsDataStorage/Source/TypedElementsDataStorage/Private/UI/TypedElementSlateWidgetReferenceColumnUpdateProcessor.cpp
@@ -11,7 +11,6 @@ void UTypedElementSlateWidgetReferenceColumnUpdateFactory::RegisterQueries(IType
 	RegisterDeleteColumnOnWidgetDeleteQuery(DataStorage);
 }
 
-#pragma optimize("", off)
 void UTypedElementSlateWidgetReferenceColumnUpdateFactory::RegisterDeleteRowOnWidgetDeleteQuery(ITypedElementDataStorageInterface& DataStorage) const
 {
 	using namespace TypedElementQueryBuilder;
@@ -34,7 +33,6 @@ void UTypedElementSlateWidgetReferenceColumnUpdateFactory::RegisterDeleteRowOnWi
 		.Compile()
 	);
 }
-#pragma optimize("", on)
 
 void UTypedElementSlateWidgetReferenceColumnUpdateFactory::RegisterDeleteColumnOnWidgetDeleteQuery(ITypedElementDataStorageInterface& DataStorage) const
 {
