Update to TEDS's phase and tick group support.

The following changes were made:
- Phases and tick groups can now be registered and referenced in query callbacks. This avoids repeatedly having to setup dependencies.
- Phase processors are introduce which allow query callbacks to be run at the start and end of a phase. This provides an opportunity to make manipulations like add/remove rows/columns that would otherwise be delayed until all processing has completed.
- The Prepare* and Finalize* tick groups were removed to avoid confusion with the new phase processors.
- The remaining tick groups are now registered to have the minimally required restrictions assigned.
- Fixed an issue with None queries behaving unexpectedly.
- Removed a few asserts in favor of reporting that Mass' entity manager isn't ready yet. This is more in line with the spirit of TEDS to avoid crashing on unavailable input.
- The UTypedElementDataStorageFactory used to automatically register information with TEDS has been updated with explicit steps to register tables and tick groups. There's nothing special about these steps that can't be safely done in other steps, but this helps keep code readable and helps to indicate the type of information should be registered through these factories.
- Updated several query callbacks to use the new PhaseAmble processor.

diff --git a/Engine/Plugins/Experimental/TypedElementsDataStorage/Source/TypedElementsDataStorage/Private/Compatibility/TypedElementActorTransformProcessors.cpp b/Engine/Plugins/Experimental/TypedElementsDataStorage/Source/TypedElementsDataStorage/Private/Compatibility/TypedElementActorTransformProcessors.cpp
--- a/Engine/Plugins/Experimental/TypedElementsDataStorage/Source/TypedElementsDataStorage/Private/Compatibility/TypedElementActorTransformProcessors.cpp
+++ b/Engine/Plugins/Experimental/TypedElementsDataStorage/Source/TypedElementsDataStorage/Private/Compatibility/TypedElementActorTransformProcessors.cpp
@@ -24,8 +24,7 @@ void UTypedElementActorTransformFactory::RegisterActorAddTransformColumn(ITypedE
 		Select(
 			TEXT("Add transform column to actor"),
 			FProcessor(DSI::EQueryTickPhase::PrePhysics, 
-				DataStorage.GetQueryTickGroupName(DSI::EQueryTickGroups::PrepareSyncExternalToDataStorage))
-				.SetBeforeGroup(DataStorage.GetQueryTickGroupName(DSI::EQueryTickGroups::SyncExternalToDataStorage))
+				DataStorage.GetQueryTickGroupName(DSI::EQueryTickGroups::SyncExternalToDataStorage))
 				.ForceToGameThread(true),
 			[](DSI::IQueryContext& Context, TypedElementRowHandle Row, const FMassActorFragment& Actor)
 			{
diff --git a/Engine/Plugins/Experimental/TypedElementsDataStorage/Source/TypedElementsDataStorage/Private/Processors/TypedElementMiscProcessors.cpp b/Engine/Plugins/Experimental/TypedElementsDataStorage/Source/TypedElementsDataStorage/Private/Processors/TypedElementMiscProcessors.cpp
--- a/Engine/Plugins/Experimental/TypedElementsDataStorage/Source/TypedElementsDataStorage/Private/Processors/TypedElementMiscProcessors.cpp
+++ b/Engine/Plugins/Experimental/TypedElementsDataStorage/Source/TypedElementsDataStorage/Private/Processors/TypedElementMiscProcessors.cpp
@@ -13,9 +13,7 @@ void UTypedElementRemoveSyncToWorldTagFactory::RegisterQueries(ITypedElementData
 	DataStorage.RegisterQuery(
 		Select(
 			TEXT("Remove 'sync to world' tag"),
-			FProcessor(DSI::EQueryTickPhase::FrameEnd, 
-				DataStorage.GetQueryTickGroupName(DSI::EQueryTickGroups::FinalizeSyncDataStorageToExternal))
-				.SetAfterGroup(DataStorage.GetQueryTickGroupName(DSI::EQueryTickGroups::SyncDataStorageToExternal)),
+			FPhaseAmble(FPhaseAmble::ELocation::Postamble, DSI::EQueryTickPhase::FrameEnd),
 			[](DSI::IQueryContext& Context, const TypedElementRowHandle* Rows)
 			{
 				Context.RemoveColumns<FTypedElementSyncBackToWorldTag>(TConstArrayView<TypedElementRowHandle>(Rows, Context.GetRowCount()));
diff --git a/Engine/Plugins/Experimental/TypedElementsDataStorage/Source/TypedElementsDataStorage/Private/Processors/TypedElementProcessorAdaptors.cpp b/Engine/Plugins/Experimental/TypedElementsDataStorage/Source/TypedElementsDataStorage/Private/Processors/TypedElementProcessorAdaptors.cpp
--- a/Engine/Plugins/Experimental/TypedElementsDataStorage/Source/TypedElementsDataStorage/Private/Processors/TypedElementProcessorAdaptors.cpp
+++ b/Engine/Plugins/Experimental/TypedElementsDataStorage/Source/TypedElementsDataStorage/Private/Processors/TypedElementProcessorAdaptors.cpp
@@ -269,6 +269,38 @@ struct FMassDirectContextForwarder final : public ITypedElementDataStorageInterf
 };
 
 
+
+/**
+ * FPhasePreOrPostAmbleExecutor
+ */
+FPhasePreOrPostAmbleExecutor::FPhasePreOrPostAmbleExecutor(FMassEntityManager& EntityManager, float DeltaTime)
+	: Context(EntityManager, DeltaTime)
+{
+	Context.SetDeferredCommandBuffer(MakeShared<FMassCommandBuffer>());
+}
+
+FPhasePreOrPostAmbleExecutor::~FPhasePreOrPostAmbleExecutor()
+{
+	Context.FlushDeferred();
+}
+
+void FPhasePreOrPostAmbleExecutor::ExecuteQuery(ITypedElementDataStorageInterface::FQueryDescription& Description, FMassEntityQuery& NativeQuery,
+	ITypedElementDataStorageInterface::QueryCallbackRef Callback)
+{
+	NativeQuery.ForEachEntityChunk(Context.GetEntityManagerChecked(), Context,
+		[&Callback, &Description](FMassExecutionContext& ExecutionContext)
+		{
+			if (FTypedElementQueryProcessorData::PrepareCachedDependenciesOnQuery(Description, ExecutionContext))
+			{
+				FMassContextForwarder QueryContext(ExecutionContext);
+				Callback(Description, QueryContext);
+			}
+		}
+	);
+}
+
+
+
 /**
  * FTypedElementQueryProcessorData
  */
@@ -277,7 +309,7 @@ FTypedElementQueryProcessorData::FTypedElementQueryProcessorData(UMassProcessor&
 {
 }
 
-EMassProcessingPhase FTypedElementQueryProcessorData::MapToMassProcessingPhase(ITypedElementDataStorageInterface::EQueryTickPhase Phase) const
+EMassProcessingPhase FTypedElementQueryProcessorData::MapToMassProcessingPhase(ITypedElementDataStorageInterface::EQueryTickPhase Phase)
 {
 	switch(Phase)
 	{
@@ -300,7 +332,7 @@ FString FTypedElementQueryProcessorData::GetProcessorName() const
 	return ParentQuery ? ParentQuery->Description.Callback.Name.ToString() : FString(TEXT("<unnamed>"));
 }
 
-bool FTypedElementQueryProcessorData::PrepareCachedDependenciesOnParentQuery(
+bool FTypedElementQueryProcessorData::PrepareCachedDependenciesOnQuery(
 	ITypedElementDataStorageInterface::FQueryDescription& Description, FMassExecutionContext& Context)
 {
 	const int32 DependencyCount = Description.DependencyTypes.Num();
@@ -345,7 +377,7 @@ ITypedElementDataStorageInterface::FQueryResult FTypedElementQueryProcessorData:
 	NativeQuery.ForEachEntityChunk(EntityManager, Context,
 		[&Result, &Callback, &Description](FMassExecutionContext& Context)
 		{
-			if (PrepareCachedDependenciesOnParentQuery(Description, Context))
+			if (PrepareCachedDependenciesOnQuery(Description, Context))
 			{
 				FMassDirectContextForwarder QueryContext(Context);
 				Callback(Description, QueryContext);
@@ -368,7 +400,7 @@ void FTypedElementQueryProcessorData::Execute(FMassEntityManager& EntityManager,
 	Query.ForEachEntityChunk(EntityManager, Context,
 		[&Description](FMassExecutionContext& Context)
 		{
-			if (PrepareCachedDependenciesOnParentQuery(Description, Context))
+			if (PrepareCachedDependenciesOnQuery(Description, Context))
 			{
 				FMassContextForwarder QueryContext(Context);
 				Description.Callback.Function(Description, QueryContext);
@@ -402,14 +434,8 @@ void UTypedElementQueryProcessorCallbackAdapterProcessor::ConfigureQueryCallback
 	bRequiresGameThreadExecution = TargetParentQuery.Description.Callback.bForceToGameThread;
 	ExecutionFlags = static_cast<int32>(EProcessorExecutionFlags::Editor); 
 	ExecutionOrder.ExecuteInGroup = TargetParentQuery.Description.Callback.Group;
-	if (!TargetParentQuery.Description.Callback.BeforeGroup.IsNone())
-	{
-		ExecutionOrder.ExecuteBefore.Add(TargetParentQuery.Description.Callback.BeforeGroup);
-	}
-	if (!TargetParentQuery.Description.Callback.AfterGroup.IsNone())
-	{
-		ExecutionOrder.ExecuteAfter.Add(TargetParentQuery.Description.Callback.AfterGroup);
-	}
+	ExecutionOrder.ExecuteBefore = TargetParentQuery.Description.Callback.BeforeGroups;
+	ExecutionOrder.ExecuteAfter = TargetParentQuery.Description.Callback.AfterGroups;
 	ProcessingPhase = Data.MapToMassProcessingPhase(TargetParentQuery.Description.Callback.Phase);
 
 	Super::PostInitProperties();
diff --git a/Engine/Plugins/Experimental/TypedElementsDataStorage/Source/TypedElementsDataStorage/Private/Processors/TypedElementProcessorAdaptors.h b/Engine/Plugins/Experimental/TypedElementsDataStorage/Source/TypedElementsDataStorage/Private/Processors/TypedElementProcessorAdaptors.h
--- a/Engine/Plugins/Experimental/TypedElementsDataStorage/Source/TypedElementsDataStorage/Private/Processors/TypedElementProcessorAdaptors.h
+++ b/Engine/Plugins/Experimental/TypedElementsDataStorage/Source/TypedElementsDataStorage/Private/Processors/TypedElementProcessorAdaptors.h
@@ -2,7 +2,9 @@
 
 #pragma once
 
+#include "Containers/ArrayView.h"
 #include "Elements/Interfaces/TypedElementDataStorageInterface.h"
+#include "MassExecutionContext.h"
 #include "MassObserverProcessor.h"
 #include "MassProcessor.h"
 #include "UObject/ObjectMacros.h"
@@ -11,6 +13,17 @@
 
 struct FTypedElementDatabaseExtendedQuery;
 
+struct FPhasePreOrPostAmbleExecutor
+{
+	FPhasePreOrPostAmbleExecutor(FMassEntityManager& EntityManager, float DeltaTime);
+	~FPhasePreOrPostAmbleExecutor();
+
+	void ExecuteQuery(ITypedElementDataStorageInterface::FQueryDescription& Description, FMassEntityQuery& NativeQuery,
+		ITypedElementDataStorageInterface::QueryCallbackRef Callback);
+
+	FMassExecutionContext Context;
+};
+
 USTRUCT()
 struct FTypedElementQueryProcessorData
 {
@@ -19,7 +32,7 @@ struct FTypedElementQueryProcessorData
 	FTypedElementQueryProcessorData() = default;
 	explicit FTypedElementQueryProcessorData(UMassProcessor& Owner);
 
-	EMassProcessingPhase MapToMassProcessingPhase(ITypedElementDataStorageInterface::EQueryTickPhase Phase) const;
+	static EMassProcessingPhase MapToMassProcessingPhase(ITypedElementDataStorageInterface::EQueryTickPhase Phase);
 	FString GetProcessorName() const;
 
 	static ITypedElementDataStorageInterface::FQueryResult Execute(
@@ -29,12 +42,11 @@ struct FTypedElementQueryProcessorData
 		FMassEntityManager& EntityManager);
 	void Execute(FMassEntityManager& EntityManager, FMassExecutionContext& Context);
 
+	static bool PrepareCachedDependenciesOnQuery(
+		ITypedElementDataStorageInterface::FQueryDescription& Description, FMassExecutionContext& Context);
+
 	FTypedElementDatabaseExtendedQuery* ParentQuery{ nullptr };
 	FMassEntityQuery Query;
-
-private:
-	static bool PrepareCachedDependenciesOnParentQuery(
-		ITypedElementDataStorageInterface::FQueryDescription& Description, FMassExecutionContext& Context);
 };
 
 /**
diff --git a/Engine/Plugins/Experimental/TypedElementsDataStorage/Source/TypedElementsDataStorage/Private/TypedElementDatabase.cpp b/Engine/Plugins/Experimental/TypedElementsDataStorage/Source/TypedElementsDataStorage/Private/TypedElementDatabase.cpp
--- a/Engine/Plugins/Experimental/TypedElementsDataStorage/Source/TypedElementsDataStorage/Private/TypedElementDatabase.cpp
+++ b/Engine/Plugins/Experimental/TypedElementsDataStorage/Source/TypedElementsDataStorage/Private/TypedElementDatabase.cpp
@@ -13,13 +13,7 @@
 #include "Stats/Stats2.h"
 #include "TickTaskManagerInterface.h"
 
-const FName UTypedElementDatabase::TickGroupName_PrepareSyncWorldToMass(TEXT("PrepareSyncWorldToMass"));
-const FName UTypedElementDatabase::TickGroupName_FinalizeSyncWorldToMass(TEXT("FinalizeSyncWorldToMass"));
-const FName UTypedElementDatabase::TickGroupName_PrepareSyncMassToExternal(TEXT("PrepareSyncMassToExternal"));
-const FName UTypedElementDatabase::TickGroupName_FinalizeSyncMassToExternal(TEXT("FinalizeSyncMassToExternal"));
-const FName UTypedElementDatabase::TickGroupName_PrepareSyncWidget(TEXT("PrepareSyncWidgets"));
 const FName UTypedElementDatabase::TickGroupName_SyncWidget(TEXT("SyncWidgets"));
-const FName UTypedElementDatabase::TickGroupName_FinalizeSyncWidget(TEXT("FinalizeSyncWidgets"));
 
 void UTypedElementDatabase::Initialize()
 {
@@ -29,6 +23,33 @@ void UTypedElementDatabase::Initialize()
 	Mass->GetOnPreTickDelegate().AddUObject(this, &UTypedElementDatabase::OnPreMassTick);
 
 	ActiveEditorEntityManager = Mass->GetMutableEntityManager();
+	ActiveEditorPhaseManager = Mass->GetMutablePhaseManager();
+
+	using PhaseType = std::underlying_type_t<EQueryTickPhase>;
+	for (PhaseType PhaseId = 0; PhaseId < static_cast<PhaseType>(EQueryTickPhase::Max); ++PhaseId)
+	{
+		EQueryTickPhase Phase = static_cast<EQueryTickPhase>(PhaseId);
+		EMassProcessingPhase MassPhase = FTypedElementQueryProcessorData::MapToMassProcessingPhase(Phase);
+		
+		ActiveEditorPhaseManager->GetOnPhaseStart(MassPhase).AddLambda(
+			[this, Phase](float DeltaTime)
+			{
+				PreparePhase(Phase, DeltaTime);
+			});
+
+		ActiveEditorPhaseManager->GetOnPhaseEnd(MassPhase).AddLambda(
+			[this, Phase](float DeltaTime)
+			{
+				FinalizePhase(Phase, DeltaTime);
+			});
+
+		// Guarantee that syncing to the data storage always happens before syncing to external.
+		RegisterTickGroup(GetQueryTickGroupName(EQueryTickGroups::SyncExternalToDataStorage), 
+			Phase, GetQueryTickGroupName(EQueryTickGroups::SyncDataStorageToExternal), {}, false);
+		// Guarantee that widgets syncs happen after external data has been updated to the data storage.
+		RegisterTickGroup(GetQueryTickGroupName(EQueryTickGroups::SyncWidgets),
+			Phase, {}, GetQueryTickGroupName(EQueryTickGroups::SyncExternalToDataStorage), false);
+	}
 }
 
 void UTypedElementDatabase::Deinitialize()
@@ -261,6 +282,32 @@ ColumnDataResult UTypedElementDatabase::GetColumnData(TypedElementRowHandle Row,
 	return ColumnDataResult{ nullptr, nullptr };
 }
 
+void UTypedElementDatabase::RegisterTickGroup(
+	FName GroupName, EQueryTickPhase Phase, FName BeforeGroup, FName AfterGroup, bool bRequiresMainThread)
+{
+	FTickGroupDescription& Group = TickGroupDescriptions.FindOrAdd({ GroupName, Phase });
+
+	if (!Group.BeforeGroups.Find(BeforeGroup))
+	{
+		Group.BeforeGroups.Add(BeforeGroup);
+	}
+
+	if (!Group.AfterGroups.Find(AfterGroup))
+	{
+		Group.AfterGroups.Add(AfterGroup);
+	}
+
+	if (bRequiresMainThread)
+	{
+		Group.bRequiresMainThread = true;
+	}
+}
+
+void UTypedElementDatabase::UnregisterTickGroup(FName GroupName, EQueryTickPhase Phase)
+{
+	TickGroupDescriptions.Remove({ GroupName, Phase });
+}
+
 TypedElementQueryHandle UTypedElementDatabase::RegisterQuery(FQueryDescription&& Query)
 {
 	auto LocalToNativeAccess = [](EQueryAccessType Access) -> EMassFragmentAccess
@@ -301,6 +348,10 @@ TypedElementQueryHandle UTypedElementDatabase::RegisterQuery(FQueryDescription&&
 					StoredQuery.Processor.Reset(Observer);
 					return Observer->GetQuery();
 				}
+				case EQueryCallbackType::PhasePreparation:
+					break;
+				case EQueryCallbackType::PhaseFinalization:
+					break;
 				default:
 					checkf(false, TEXT("Unsupported query callback type %i."), static_cast<int>(Query.Callback.Type));
 					break;
@@ -314,12 +365,13 @@ TypedElementQueryHandle UTypedElementDatabase::RegisterQuery(FQueryDescription&&
 	
 	FMassEntityQuery& NativeQuery = SetupNativeQuery(Query, StoredQuery);
 
+	// Setup selected columns section
 	if (Query.Action == FQueryDescription::EActionType::Count)
 	{
 		checkf(Query.SelectionTypes.IsEmpty(), TEXT("Count queries for the Typed Elements Data Storage can't have entries for selection."));
 		checkf(Query.SelectionAccessTypes.IsEmpty(), TEXT("Count queries for the Typed Elements Data Storage can't have entries for selection."));
 	}
-	else
+	else if (Query.Action == FQueryDescription::EActionType::Select)
 	{
 		const int32 SelectionCount = Query.SelectionTypes.Num();
 		checkf(SelectionCount == Query.SelectionAccessTypes.Num(),
@@ -338,7 +390,12 @@ TypedElementQueryHandle UTypedElementDatabase::RegisterQuery(FQueryDescription&&
 			NativeQuery.AddRequirement(Type.Get(), LocalToNativeAccess(Query.SelectionAccessTypes[SelectionIndex]));
 		}
 	}
+	else
+	{
+		checkf(Query.Action == FQueryDescription::EActionType::None, TEXT("Unexpected query action: %i."), static_cast<int32>(Query.Action));
+	}
 
+	// Configure conditions
 	if (Query.bSimpleQuery) // This backend currently only supports simple queries.
 	{
 		checkf(Query.ConditionTypes.Num() == Query.ConditionOperators.Num(),
@@ -376,6 +433,7 @@ TypedElementQueryHandle UTypedElementDatabase::RegisterQuery(FQueryDescription&&
 		}
 	}
 
+	// Assign dependencies.
 	const int32 DependencyCount = Query.DependencyTypes.Num();
 	checkf(DependencyCount == Query.DependencyFlags.Num() && DependencyCount == Query.CachedDependencies.Num(),
 		TEXT("The number of query depedencies (%i) doesn't match the number of dependency access types (%i) and/or cached dependencies count (%i)."),
@@ -394,8 +452,46 @@ TypedElementQueryHandle UTypedElementDatabase::RegisterQuery(FQueryDescription&&
 			EnumHasAllFlags(Flags, EQueryDependencyFlags::GameThreadBound));
 	}
 
+	// Copy pre-registered phase and group information.
+	const FTickGroupDescription* TickGroup = TickGroupDescriptions.Find({ Query.Callback.Group, Query.Callback.Phase });
+	if (TickGroup)
+	{
+		for (auto It = Query.Callback.BeforeGroups.CreateIterator(); It; ++It)
+		{
+			if (TickGroup->BeforeGroups.Contains(*It))
+			{
+				It.RemoveCurrentSwap();
+			}
+		}
+		Query.Callback.BeforeGroups.Append(TickGroup->BeforeGroups);
+		for (auto It = Query.Callback.AfterGroups.CreateIterator(); It; ++It)
+		{
+			if (TickGroup->AfterGroups.Contains(*It))
+			{
+				It.RemoveCurrentSwap();
+			}
+		}
+		Query.Callback.AfterGroups.Append(TickGroup->AfterGroups);
+		if (TickGroup->bRequiresMainThread)
+		{
+			Query.Callback.bForceToGameThread = true;
+		}
+	}
+
+	// Register Phase processors locally.
+	switch (Query.Callback.Type)
+	{
+	case EQueryCallbackType::PhasePreparation:
+		PhasePreparationQueries[static_cast<std::underlying_type_t<EQueryTickPhase>>(Query.Callback.Phase)].Add(Result.Handle);
+		break;
+	case EQueryCallbackType::PhaseFinalization:
+		PhaseFinalizationQueries[static_cast<std::underlying_type_t<EQueryTickPhase>>(Query.Callback.Phase)].Add(Result.Handle);
+		break;
+	}
+
 	StoredQuery.Description = MoveTemp(Query);
 
+	// Register regular processors and observer with Mass.
 	if (StoredQuery.Processor)
 	{
 		if (StoredQuery.Processor->IsA<UTypedElementQueryProcessorCallbackAdapterProcessor>())
@@ -452,6 +548,22 @@ void UTypedElementDatabase::UnregisterQuery(TypedElementQueryHandle Query)
 					*QueryData.Description.Callback.Name.ToString(), *QueryData.Processor->GetSparseClassDataStruct()->GetName());
 			}
 		}
+		else if (QueryData.Description.Callback.Type == EQueryCallbackType::PhasePreparation)
+		{
+			int32 Index;
+			if (PhasePreparationQueries[static_cast<std::underlying_type_t<EQueryTickPhase>>(QueryData.Description.Callback.Phase)].Find(Query, Index))
+			{
+				PhasePreparationQueries[static_cast<std::underlying_type_t<EQueryTickPhase>>(QueryData.Description.Callback.Phase)].RemoveAt(Index);
+			}
+		}
+		else if (QueryData.Description.Callback.Type == EQueryCallbackType::PhaseFinalization)
+		{
+			int32 Index;
+			if (PhaseFinalizationQueries[static_cast<std::underlying_type_t<EQueryTickPhase>>(QueryData.Description.Callback.Phase)].Find(Query, Index))
+			{
+				PhaseFinalizationQueries[static_cast<std::underlying_type_t<EQueryTickPhase>>(QueryData.Description.Callback.Phase)].RemoveAt(Index);
+			}
+		}
 		else
 		{
 			QueryData.NativeQuery.Clear();
@@ -484,27 +596,12 @@ FName UTypedElementDatabase::GetQueryTickGroupName(EQueryTickGroups Group) const
 	{
 		case EQueryTickGroups::Default:
 			return NAME_None;
-		case EQueryTickGroups::PrepareSyncExternalToDataStorage:
-			return TickGroupName_PrepareSyncWorldToMass;
 		case EQueryTickGroups::SyncExternalToDataStorage:
 			return UE::Mass::ProcessorGroupNames::SyncWorldToMass;
-		case EQueryTickGroups::FinalizeSyncExternalToDataStorage:
-			return TickGroupName_FinalizeSyncWorldToMass;
-		
-		case EQueryTickGroups::PrepareSyncDataStorageToExternal:
-			return TickGroupName_PrepareSyncMassToExternal;
 		case EQueryTickGroups::SyncDataStorageToExternal:
 			return UE::Mass::ProcessorGroupNames::UpdateWorldFromMass;
-		case EQueryTickGroups::FinalizeSyncDataStorageToExternal:
-			return TickGroupName_FinalizeSyncMassToExternal;
-		
-		case EQueryTickGroups::PrepareSyncWidgets:
-			return TickGroupName_PrepareSyncWidget;
 		case EQueryTickGroups::SyncWidgets:
 			return TickGroupName_SyncWidget;
-		case EQueryTickGroups::FinalizeSyncWidgets:
-			return TickGroupName_FinalizeSyncWidget;
-
 		default:
 			checkf(false, TEXT("EQueryTickGroups value %i can't be translated to a group name by this Data Storage backend."), static_cast<int>(Group));
 			return NAME_None;
@@ -533,10 +630,15 @@ ITypedElementDataStorageInterface::FQueryResult UTypedElementDatabase::RunQuery(
 				case FQueryDescription::EActionType::Select:
 					// Fallthrough: There's nothing to callback to, so only return the total count.
 				case FQueryDescription::EActionType::Count:
-					checkf(ActiveEditorEntityManager, 
-						TEXT("Unable to run Typed Element Data Storage query before a MASS Entity Manager has been assigned."));
-					Result.Count = QueryData.NativeQuery.GetNumMatchingEntities(*ActiveEditorEntityManager);
-					Result.Completed = FQueryResult::ECompletion::Fully;
+					if (ActiveEditorEntityManager)
+					{
+						Result.Count = QueryData.NativeQuery.GetNumMatchingEntities(*ActiveEditorEntityManager);
+						Result.Completed = FQueryResult::ECompletion::Fully;
+					}
+					else
+					{
+						Result.Completed = FQueryResult::ECompletion::Unavailable;
+					}
 					break;
 				default:
 					Result.Completed = FQueryResult::ECompletion::Unsupported;
@@ -578,17 +680,34 @@ ITypedElementDataStorageInterface::FQueryResult UTypedElementDatabase::RunQuery(
 				Result.Completed = FQueryResult::ECompletion::Fully;
 				break;
 			case FQueryDescription::EActionType::Select:
-				checkf(ActiveEditorEntityManager,
-					TEXT("Unable to run Typed Element Data Storage query before a MASS Entity Manager has been assigned."));
-				Result = FTypedElementQueryProcessorData::Execute(
-					Callback, QueryData.Description, QueryData.NativeQuery, *ActiveEditorEntityManager);
+				if (ActiveEditorEntityManager)
+				{
+					if (!QueryData.Processor.IsValid())
+					{
+						Result = FTypedElementQueryProcessorData::Execute(
+							Callback, QueryData.Description, QueryData.NativeQuery, *ActiveEditorEntityManager);
+					}
+					else
+					{
+						Result.Completed = FQueryResult::ECompletion::Unsupported;
+					}
+				}
+				else
+				{
+					Result.Completed = FQueryResult::ECompletion::Unavailable;
+				}
 				break;
 			case FQueryDescription::EActionType::Count:
 				// Only the count is requested so no need to trigger the callback.
-				checkf(ActiveEditorEntityManager,
-					TEXT("Unable to run Typed Element Data Storage query before a MASS Entity Manager has been assigned."));
-				Result.Count = QueryData.NativeQuery.GetNumMatchingEntities(*ActiveEditorEntityManager);
-				Result.Completed = FQueryResult::ECompletion::Fully;
+				if (ActiveEditorEntityManager)
+				{
+					Result.Count = QueryData.NativeQuery.GetNumMatchingEntities(*ActiveEditorEntityManager);
+					Result.Completed = FQueryResult::ECompletion::Fully;
+				}
+				else
+				{
+					Result.Completed = FQueryResult::ECompletion::Unavailable;
+				}
 				break;
 			default:
 				Result.Completed = FQueryResult::ECompletion::Unsupported;
@@ -628,9 +747,34 @@ void* UTypedElementDatabase::GetExternalSystemAddress(UClass* Target)
 	return nullptr;
 }
 
+void UTypedElementDatabase::PreparePhase(EQueryTickPhase Phase, float DeltaTime)
+{
+	PhasePreOrPostAmble(Phase, DeltaTime, PhasePreparationQueries[static_cast<std::underlying_type_t<EQueryTickPhase>>(Phase)]);
+}
+
+void UTypedElementDatabase::FinalizePhase(EQueryTickPhase Phase, float DeltaTime)
+{
+	PhasePreOrPostAmble(Phase, DeltaTime, PhaseFinalizationQueries[static_cast<std::underlying_type_t<EQueryTickPhase>>(Phase)]);
+}
+
+void UTypedElementDatabase::PhasePreOrPostAmble(EQueryTickPhase Phase, float DeltaTime, TArray<TypedElementQueryHandle>& QueryHandles)
+{
+	if (ActiveEditorEntityManager && !QueryHandles.IsEmpty())
+	{
+		FPhasePreOrPostAmbleExecutor Executor(*ActiveEditorEntityManager, DeltaTime);
+		for (TypedElementQueryHandle Query : QueryHandles)
+		{
+			QueryStore::Handle Handle;
+			Handle.Handle = Query;
+			FTypedElementDatabaseExtendedQuery& QueryData = Queries.GetMutable(Handle);
+			Executor.ExecuteQuery(QueryData.Description, QueryData.NativeQuery, QueryData.Description.Callback.Function);
+		}
+	}
+}
+
 void UTypedElementDatabase::Reset()
 {
 	Tables.Reset();
 	TableNameLookup.Reset();
 	ActiveEditorEntityManager.Reset();
-}
\ No newline at end of file
+}
diff --git a/Engine/Plugins/Experimental/TypedElementsDataStorage/Source/TypedElementsDataStorage/Private/TypedElementDatabase.h b/Engine/Plugins/Experimental/TypedElementsDataStorage/Source/TypedElementsDataStorage/Private/TypedElementDatabase.h
--- a/Engine/Plugins/Experimental/TypedElementsDataStorage/Source/TypedElementsDataStorage/Private/TypedElementDatabase.h
+++ b/Engine/Plugins/Experimental/TypedElementsDataStorage/Source/TypedElementsDataStorage/Private/TypedElementDatabase.h
@@ -14,6 +14,7 @@
 #include "TypedElementDatabase.generated.h"
 
 struct FMassEntityManager;
+struct FMassProcessingPhaseManager;
 class UWorld;
 
 struct FTypedElementDatabaseExtendedQuery
@@ -64,6 +65,9 @@ public:
 		TConstArrayView<TypedElement::ColumnUtils::Argument> Arguments) override;
 	ColumnDataResult GetColumnData(TypedElementRowHandle Row, FTopLevelAssetPath ColumnName) override;
 
+	void RegisterTickGroup(FName GroupName, EQueryTickPhase Phase, FName BeforeGroup, FName AfterGroup, bool bRequiresMainThread);
+	void UnregisterTickGroup(FName GroupName, EQueryTickPhase Phase);
+
 	TypedElementQueryHandle RegisterQuery(FQueryDescription&& Query) override;
 	void UnregisterQuery(TypedElementQueryHandle Query) override;
 	const FQueryDescription& GetQueryDescription(TypedElementQueryHandle Query) const override;
@@ -78,22 +82,41 @@ public:
 private:
 	using QueryStore = TTypedElementHandleStore<FTypedElementDatabaseExtendedQuery>;
 
+	struct FTickGroupId
+	{
+		FName Name;
+		EQueryTickPhase Phase;
+
+		friend inline uint32 GetTypeHash(const FTickGroupId& Id){ return HashCombine(GetTypeHash(Id.Name), GetTypeHash(Id.Phase)); }
+		friend inline bool operator==(const FTickGroupId& Lhs, const FTickGroupId& Rhs) { return Lhs.Phase == Rhs.Phase && Lhs.Name == Rhs.Name; }
+		friend inline bool operator!=(const FTickGroupId& Lhs, const FTickGroupId& Rhs) { return Lhs.Phase != Rhs.Phase || Lhs.Name != Rhs.Name; }
+	};
+
+	struct FTickGroupDescription
+	{
+		TArray<FName> BeforeGroups;
+		TArray<FName> AfterGroups;
+		bool bRequiresMainThread{ false };
+	};
+	
+	void PreparePhase(EQueryTickPhase Phase, float DeltaTime);
+	void FinalizePhase(EQueryTickPhase Phase, float DeltaTime);
+	void PhasePreOrPostAmble(EQueryTickPhase Phase, float DeltaTime, TArray<TypedElementQueryHandle>& Queries);
 	void Reset();
 	
-	static const FName TickGroupName_PrepareSyncWorldToMass;
-	static const FName TickGroupName_FinalizeSyncWorldToMass;
-	static const FName TickGroupName_PrepareSyncMassToExternal;
-	static const FName TickGroupName_FinalizeSyncMassToExternal;
-	static const FName TickGroupName_PrepareSyncWidget;
 	static const FName TickGroupName_SyncWidget;
-	static const FName TickGroupName_FinalizeSyncWidget;
 	
 	TArray<FMassArchetypeHandle> Tables;
 	TMap<FName, TypedElementTableHandle> TableNameLookup;
+	TMap<FTickGroupId, FTickGroupDescription> TickGroupDescriptions;
+
+	TArray<TypedElementQueryHandle> PhasePreparationQueries[static_cast<std::underlying_type_t<EQueryTickPhase>>(EQueryTickPhase::Max)];
+	TArray<TypedElementQueryHandle> PhaseFinalizationQueries[static_cast<std::underlying_type_t<EQueryTickPhase>>(EQueryTickPhase::Max)];
 
 	QueryStore Queries;
 
 	FTypedElementOnDataStorageUpdate OnUpdateDelegate;
 
 	TSharedPtr<FMassEntityManager> ActiveEditorEntityManager;
+	TSharedPtr<FMassProcessingPhaseManager> ActiveEditorPhaseManager;
 };
\ No newline at end of file
diff --git a/Engine/Plugins/Experimental/TypedElementsDataStorage/Source/TypedElementsDataStorage/Private/UI/TypedElementSlateWidgetReferenceColumnUpdateProcessor.cpp b/Engine/Plugins/Experimental/TypedElementsDataStorage/Source/TypedElementsDataStorage/Private/UI/TypedElementSlateWidgetReferenceColumnUpdateProcessor.cpp
--- a/Engine/Plugins/Experimental/TypedElementsDataStorage/Source/TypedElementsDataStorage/Private/UI/TypedElementSlateWidgetReferenceColumnUpdateProcessor.cpp
+++ b/Engine/Plugins/Experimental/TypedElementsDataStorage/Source/TypedElementsDataStorage/Private/UI/TypedElementSlateWidgetReferenceColumnUpdateProcessor.cpp
@@ -19,7 +19,7 @@ void UTypedElementSlateWidgetReferenceColumnUpdateFactory::RegisterDeleteRowOnWi
 	DataStorage.RegisterQuery(
 		Select(
 			TEXT("Delete row with deleted widget"),
-			FProcessor(DSI::EQueryTickPhase::FrameEnd, DataStorage.GetQueryTickGroupName(DSI::EQueryTickGroups::PrepareSyncWidgets)).ForceToGameThread(true),
+			FPhaseAmble(FPhaseAmble::ELocation::Preamble, DSI::EQueryTickPhase::FrameEnd),
 			[](DSI::IQueryContext& Context, TypedElementRowHandle Row, const FTypedElementSlateWidgetReferenceColumn& WidgetReference)
 			{
 				if (!WidgetReference.Widget.IsValid())
@@ -42,7 +42,7 @@ void UTypedElementSlateWidgetReferenceColumnUpdateFactory::RegisterDeleteColumnO
 	DataStorage.RegisterQuery(
 		Select(
 			TEXT("Delete widget column for deleted widget"),
-			FProcessor(DSI::EQueryTickPhase::FrameEnd, DataStorage.GetQueryTickGroupName(DSI::EQueryTickGroups::PrepareSyncWidgets)).ForceToGameThread(true),
+			FPhaseAmble(FPhaseAmble::ELocation::Preamble, DSI::EQueryTickPhase::FrameEnd),
 			[](DSI::IQueryContext& Context, TypedElementRowHandle Row, const FTypedElementSlateWidgetReferenceColumn& WidgetReference)
 			{
 				if (!WidgetReference.Widget.IsValid())
diff --git a/Engine/Source/Runtime/TypedElementFramework/Private/Elements/Framework/TypedElementQueryBuilder.cpp b/Engine/Source/Runtime/TypedElementFramework/Private/Elements/Framework/TypedElementQueryBuilder.cpp
--- a/Engine/Source/Runtime/TypedElementFramework/Private/Elements/Framework/TypedElementQueryBuilder.cpp
+++ b/Engine/Source/Runtime/TypedElementFramework/Private/Elements/Framework/TypedElementQueryBuilder.cpp
@@ -223,6 +223,33 @@ namespace TypedElementQueryBuilder
 	}
 
 
+	/**
+	 * FPhaseAmble
+	 */
+
+	FPhaseAmble::FPhaseAmble(ELocation InLocation, ITypedElementDataStorageInterface::EQueryTickPhase InPhase)
+		: Phase(InPhase)
+		, Location(InLocation)
+	{}
+
+	FPhaseAmble& FPhaseAmble::SetLocation(ELocation NewLocation)
+	{
+		Location = NewLocation;
+		return *this;
+	}
+
+	FPhaseAmble& FPhaseAmble::SetPhase(ITypedElementDataStorageInterface::EQueryTickPhase NewPhase)
+	{
+		Phase = NewPhase;
+		return *this;
+	}
+
+	FPhaseAmble& FPhaseAmble::ForceToGameThread(bool bForce)
+	{
+		bForceToGameThread = bForce;
+		return *this;
+	}
+
 
 	/**
 	 * Select
diff --git a/Engine/Source/Runtime/TypedElementFramework/Public/Elements/Framework/TypedElementQueryBuilder.h b/Engine/Source/Runtime/TypedElementFramework/Public/Elements/Framework/TypedElementQueryBuilder.h
--- a/Engine/Source/Runtime/TypedElementFramework/Public/Elements/Framework/TypedElementQueryBuilder.h
+++ b/Engine/Source/Runtime/TypedElementFramework/Public/Elements/Framework/TypedElementQueryBuilder.h
@@ -187,7 +187,7 @@ namespace TypedElementQueryBuilder
 		};
 
 		template<typename ColumnType>
-		FObserver(EEvent MonitorForEvent);
+		explicit FObserver(EEvent MonitorForEvent);
 		FObserver(EEvent MonitorForEvent, const UScriptStruct* MonitoredColumn);
 
 		FObserver& SetEvent(EEvent MonitorForEvent);
@@ -201,6 +201,24 @@ namespace TypedElementQueryBuilder
 		bool bForceToGameThread{ false };
 	};
 
+	struct TYPEDELEMENTFRAMEWORK_API FPhaseAmble final : public FQueryCallbackType
+	{
+		enum class ELocation : uint8
+		{
+			Preamble,
+			Postamble
+		};
+
+		FPhaseAmble(ELocation InLocation, ITypedElementDataStorageInterface::EQueryTickPhase InPhase);
+		FPhaseAmble& SetLocation(ELocation NewLocation);
+		FPhaseAmble& SetPhase(ITypedElementDataStorageInterface::EQueryTickPhase NewPhase);
+		FPhaseAmble& ForceToGameThread(bool bForce);
+
+		ITypedElementDataStorageInterface::EQueryTickPhase Phase;
+		ELocation Location;
+		bool bForceToGameThread{ false };
+	};
+
 	// Because this is a thin wrapper called from within a query callback, it's better to inline fully so all
 	// function pre/postambles can be optimized away.
 	struct FQueryContextForwarder : public ITypedElementDataStorageInterface::IQueryContext
@@ -293,4 +311,4 @@ namespace TypedElementQueryBuilder
 	};
 } // namespace TypedElementQueryBuilder
 
-#include "Elements/Framework/TypedElementQueryBuilder.inl"
\ No newline at end of file
+#include "Elements/Framework/TypedElementQueryBuilder.inl"
diff --git a/Engine/Source/Runtime/TypedElementFramework/Public/Elements/Framework/TypedElementQueryBuilder.inl b/Engine/Source/Runtime/TypedElementFramework/Public/Elements/Framework/TypedElementQueryBuilder.inl
--- a/Engine/Source/Runtime/TypedElementFramework/Public/Elements/Framework/TypedElementQueryBuilder.inl
+++ b/Engine/Source/Runtime/TypedElementFramework/Public/Elements/Framework/TypedElementQueryBuilder.inl
@@ -916,8 +916,14 @@ namespace TypedElementQueryBuilder
 			Query.Callback.Type = ITypedElementDataStorageInterface::EQueryCallbackType::Processor;
 			Query.Callback.Phase = Processor.Phase;
 			Query.Callback.Group = Processor.Group;
-			Query.Callback.BeforeGroup = Processor.BeforeGroup;
-			Query.Callback.AfterGroup = Processor.AfterGroup;
+			if (!Processor.BeforeGroup.IsNone())
+			{
+				Query.Callback.BeforeGroups.Add(Processor.BeforeGroup);
+			}
+			if (!Processor.AfterGroup.IsNone())
+			{
+				Query.Callback.AfterGroups.Add(Processor.AfterGroup);
+			}
 			Query.Callback.bForceToGameThread = Processor.bForceToGameThread;
 		}
 
@@ -935,6 +941,21 @@ namespace TypedElementQueryBuilder
 			Query.Callback.MonitoredType = Observer.Monitor;
 			Query.Callback.bForceToGameThread = Observer.bForceToGameThread;
 		}
+
+		inline void PrepareForQueryBinding(QueryDescription& Query, const FPhaseAmble& PhaseAmble)
+		{
+			switch (PhaseAmble.Location)
+			{
+			case FPhaseAmble::ELocation::Preamble:
+				Query.Callback.Type = ITypedElementDataStorageInterface::EQueryCallbackType::PhasePreparation;
+				break;
+			case FPhaseAmble::ELocation::Postamble:
+				Query.Callback.Type = ITypedElementDataStorageInterface::EQueryCallbackType::PhaseFinalization;
+				break;
+			}
+			Query.Callback.Phase = PhaseAmble.Phase;
+			Query.Callback.bForceToGameThread = PhaseAmble.bForceToGameThread;
+		}
 		
 		template<typename CallbackType, typename Function>
 		void PrepareForQueryBinding(Select& Target, QueryDescription& Query, FName Name,
diff --git a/Engine/Source/Runtime/TypedElementFramework/Public/Elements/Interfaces/TypedElementDataStorageFactory.h b/Engine/Source/Runtime/TypedElementFramework/Public/Elements/Interfaces/TypedElementDataStorageFactory.h
--- a/Engine/Source/Runtime/TypedElementFramework/Public/Elements/Interfaces/TypedElementDataStorageFactory.h
+++ b/Engine/Source/Runtime/TypedElementFramework/Public/Elements/Interfaces/TypedElementDataStorageFactory.h
@@ -19,6 +19,8 @@ class TYPEDELEMENTFRAMEWORK_API UTypedElementDataStorageFactory : public UObject
 public:
 	~UTypedElementDataStorageFactory() override = default;
 
+	virtual void RegisterTables(ITypedElementDataStorageInterface& DataStorage) const {}
+	virtual void RegisterTickGroups(ITypedElementDataStorageInterface& DataStorage) const {}
 	virtual void RegisterQueries(ITypedElementDataStorageInterface& DataStorage) const {}
 	virtual void RegisterWidgetConstructor(ITypedElementDataStorageInterface& DataStorage,
 		ITypedElementDataStorageUiInterface& DataStorageUi) const {}
diff --git a/Engine/Source/Runtime/TypedElementFramework/Public/Elements/Interfaces/TypedElementDataStorageInterface.h b/Engine/Source/Runtime/TypedElementFramework/Public/Elements/Interfaces/TypedElementDataStorageInterface.h
--- a/Engine/Source/Runtime/TypedElementFramework/Public/Elements/Interfaces/TypedElementDataStorageInterface.h
+++ b/Engine/Source/Runtime/TypedElementFramework/Public/Elements/Interfaces/TypedElementDataStorageInterface.h
@@ -154,7 +154,7 @@ public:
 		PostPhysics, //< Queries that need rigid body and cloth simulation to be completed before being executed.
 		FrameEnd, //< Catchall for queries demoted to the last possible moment.
 
-		MAX
+		Max //< Value indicating the maximum value in this enum. Not to be used as an enum value.
 	};
 
 	enum class EQueryTickGroups : uint8
@@ -162,40 +162,22 @@ public:
 		/** The standard group to run work in. */
 		Default,
 		
-		/**
-		 * Prepares to sync data from an external source to the Data Storage. This is typically used for preparing columns with
-		 * data that is needed to sync from external sources.
-		 */
-		PrepareSyncExternalToDataStorage,
 		/** 
 		 * The group for queries that need to sync data from external sources such as subsystems or the world into
 		 * the Data Storage. These typically run early in a phase.
 		 */
 		SyncExternalToDataStorage,
-		/**
-		 * Finalizes the steps to sync data from an external source to the Data Storage. This is typically needed for processing 
-		 * data with a dependency on external updates being completed.
-		 */
-		FinalizeSyncExternalToDataStorage,
-		
-		/** Prepares for data from the Data Storage to be synced to external sources such as subsystem or the world. */
-		PrepareSyncDataStorageToExternal,
 		/**
 		 * The group for quries that need to sync data from the Data Storage to external sources such as subsystems
 		 * or the world into. These typically run late in a phase.
 		 */
 		SyncDataStorageToExternal,
-		/** Finalizes the sync of data from the Data Storage to external sources such as subsystem or the world. */
-		FinalizeSyncDataStorageToExternal,
-
-		/** Prepares for widgets getting synced to/from widgets. */
-		PrepareSyncWidgets,
 		/**
 		 * Queries grouped under this name will sync data to/from widgets.
 		 */
 		SyncWidgets,
-		/** Finalize for widgets getting synced to/from widgets. */
-		FinalizeSyncWidgets
+
+		Max //< Value indicating the maximum value in this enum. Not to be used as an enum value.
 	};
 
 	enum class EQueryCallbackType : uint8
@@ -207,7 +189,22 @@ public:
 		/** The query will be run when a row is added that matches the query. The first recorded column will be actively monitored for changes. */
 		ObserveAdd,
 		/** The query will be run when a row is removed that matches the query. The first recorded column will be actively monitored for changes. */
-		ObserveRemove
+		ObserveRemove,
+		/** 
+		 * At the start of the assigned phase this query will run if there are any matches. These queries will have any deferred operations such as
+		 * adding/removing rows/columns executed before the phase starts. This introduces sync points that hinder performance and are therefore
+		 * recommended only for queries that save on work later in the phase such as repeated checks for validity.
+		 */
+		PhasePreparation,
+		/**
+		 * At the end of the assigned phase this query will run if there are any matches. These queries will have any deferred operations such as
+		 * adding/removing rows/columns executed before the phase ends. This introduces sync points that hinder performance and are therefore
+		 * recommended only cases where delaying deferred operations is not possible e.g. when tables are known to be referenced outside the update
+		 * cycle.
+		 */
+		 PhaseFinalization,
+
+		 Max //< Value indicating the maximum value in this enum. Not to be used as an enum value.
 	};
 
 	enum class EQueryAccessType : bool
@@ -341,6 +338,7 @@ public:
 
 	struct FQueryDescription;
 	using QueryCallback = TFunction<void(const FQueryDescription&, IQueryContext&)>;
+	using QueryCallbackRef = TFunctionRef<void(const FQueryDescription&, IQueryContext&)>;
 	using DirectQueryCallbackRef = TFunctionRef<void(const FQueryDescription&, IDirectQueryContext&)>;
 
 	struct FQueryDescription final
@@ -348,25 +346,30 @@ public:
 		static constexpr int32 NumInlineSelections = 8;
 		static constexpr int32 NumInlineConditions = 8;
 		static constexpr int32 NumInlineDependencies = 2;
+		static constexpr int32 NumInlineGroups = 2;
 
 		enum class EActionType : uint8
 		{
-			None,	/** Do nothing. */
-			Select,	/** Selects a set of columns for further processing. */
-			Count	/** Counts the number of entries that match the filter condition. */
+			None,	//< Do nothing.
+			Select,	//< Selects a set of columns for further processing.
+			Count,	//< Counts the number of entries that match the filter condition.
+
+			Max //< Value indicating the maximum value in this enum. Not to be used as an enum value.
 		};
 
 		using OperatorIndex = int32;
 		enum class EOperatorType : uint16
 		{
-			SimpleAll,			// Unary: Type
-			SimpleAny,			// Unary: Type
-			SimpleNone,			// Unary: Type
-			SimpleOptional,		// Unary: Type
-			And,				// Binary: left operator index, right operator index
-			Or,					// Binary: left operator index, right operator index
-			Not,				// Unary: condition index
-			Type				// Unary: Type
+			SimpleAll,			//< Unary: Type
+			SimpleAny,			//< Unary: Type
+			SimpleNone,			//< Unary: Type
+			SimpleOptional,		//< Unary: Type
+			And,				//< Binary: left operator index, right operator index
+			Or,					//< Binary: left operator index, right operator index
+			Not,				//< Unary: condition index
+			Type,				//< Unary: Type
+
+			Max //< Value indicating the maximum value in this enum. Not to be used as an enum value.
 		};
 
 		struct FBinaryOperator final
@@ -384,11 +387,11 @@ public:
 
 		struct FCallbackData
 		{
+			TArray<FName, TInlineAllocator<NumInlineGroups>> BeforeGroups;
+			TArray<FName, TInlineAllocator<NumInlineGroups>> AfterGroups;
 			QueryCallback Function;
 			FName Name;
 			FName Group;
-			FName BeforeGroup;
-			FName AfterGroup;
 			const UScriptStruct* MonitoredType{ nullptr };
 			EQueryCallbackType Type{ EQueryCallbackType::None };
 			EQueryTickPhase Phase;
