Add a console command to print all processors added by TEDS and one to print all columns supported by TEDS.

diff --git a/Engine/Plugins/Experimental/TypedElementsDataStorage/Source/TypedElementsDataStorage/Private/TypedElementDatabase.cpp b/Engine/Plugins/Experimental/TypedElementsDataStorage/Source/TypedElementsDataStorage/Private/TypedElementDatabase.cpp
--- a/Engine/Plugins/Experimental/TypedElementsDataStorage/Source/TypedElementsDataStorage/Private/TypedElementDatabase.cpp
+++ b/Engine/Plugins/Experimental/TypedElementsDataStorage/Source/TypedElementsDataStorage/Private/TypedElementDatabase.cpp
@@ -3,18 +3,54 @@
 #include "TypedElementDatabase.h"
 
 #include "Editor.h"
+#include "Elements/Framework/TypedElementRegistry.h"
 #include "Engine/World.h"
 #include "MassCommonTypes.h"
 #include "MassEntityEditorSubsystem.h"
+#include "MassEntityTypes.h"
 #include "MassSimulationSubsystem.h"
 #include "MassSubsystemAccess.h"
 #include "Processors/TypedElementProcessorAdaptors.h"
 #include "ProfilingDebugging/CpuProfilerTrace.h"
 #include "Stats/Stats2.h"
 #include "TickTaskManagerInterface.h"
+#include "UObject/UObjectIterator.h"
 
 const FName UTypedElementDatabase::TickGroupName_SyncWidget(TEXT("SyncWidgets"));
 
+FAutoConsoleCommandWithOutputDevice PrintQueryCallbacksConsoleCommand(
+	TEXT("TEDS.PrintQueryCallbacks"),
+	TEXT("Prints out a list of all processors."),
+	FConsoleCommandWithOutputDeviceDelegate::CreateLambda([](FOutputDevice& Output)
+		{
+			if (UTypedElementRegistry* Registry = UTypedElementRegistry::GetInstance())
+			{
+				if (UTypedElementDatabase* DataStorage = Cast<UTypedElementDatabase>(Registry->GetMutableDataStorage()))
+				{
+					DataStorage->DebugPrintQueryCallbacks(Output);
+				}
+			}
+		}));
+
+FAutoConsoleCommandWithOutputDevice PrintSupportedColumnsConsoleCommand(
+	TEXT("TEDS.PrintSupportedColumns"),
+	TEXT("Prints out a list of available Data Storage columns."),
+	FConsoleCommandWithOutputDeviceDelegate::CreateLambda([](FOutputDevice& Output)
+		{
+			Output.Log(TEXT("The Typed Elements Data Storage supports the following columns:"));
+			
+			UScriptStruct* FragmentTypeInfo = FMassFragment::StaticStruct();
+			UScriptStruct* TagTypeInfo = FMassTag::StaticStruct();
+			for (TObjectIterator<UScriptStruct> It; It; ++It)
+			{
+				if (It->IsChildOf(FragmentTypeInfo) || It->IsChildOf(TagTypeInfo))
+				{
+					Output.Logf(TEXT("    %s"), *It->GetFullName());
+				}
+			}
+			Output.Log(TEXT("End of Typed Elements Data Storage supported column list."));
+		}));
+
 void UTypedElementDatabase::Initialize()
 {
 	check(GEditor);
@@ -778,3 +814,39 @@ void UTypedElementDatabase::Reset()
 	TableNameLookup.Reset();
 	ActiveEditorEntityManager.Reset();
 }
+
+void UTypedElementDatabase::DebugPrintQueryCallbacks(FOutputDevice& Output)
+{
+	Output.Log(TEXT("The Typed Elements Data Storage has the following query callbacks:"));
+	Queries.ListAliveEntries(
+		[&Output](const FTypedElementDatabaseExtendedQuery& Query)
+		{
+			if (Query.Processor)
+			{
+				Output.Logf(TEXT("    [%s] %s"), 
+					IsValid(Query.Processor.Get()) ? TEXT("Valid") : TEXT("Invalid"),
+					*(Query.Processor->GetProcessorName()));
+			}
+		});
+
+	using PhaseType = std::underlying_type_t<EQueryTickPhase>;
+	for (PhaseType PhaseId = 0; PhaseId < static_cast<PhaseType>(EQueryTickPhase::Max); ++PhaseId)
+	{
+		for (TypedElementQueryHandle QueryHandle : PhasePreparationQueries[PhaseId])
+		{
+			QueryStore::Handle Handle;
+			Handle.Handle = QueryHandle;
+			const FTypedElementDatabaseExtendedQuery& QueryData = Queries.Get(Handle);
+			Output.Logf(TEXT("    [Valid] %s [Editor Phase Preamble]"), *QueryData.Description.Callback.Name.ToString());
+		}
+		for (TypedElementQueryHandle QueryHandle : PhaseFinalizationQueries[PhaseId])
+		{
+			QueryStore::Handle Handle;
+			Handle.Handle = QueryHandle;
+			const FTypedElementDatabaseExtendedQuery& QueryData = Queries.Get(Handle);
+			Output.Logf(TEXT("    [Valid] %s [Editor Phase Postamble]"), *QueryData.Description.Callback.Name.ToString());
+		}
+	}
+
+	Output.Log(TEXT("End of Typed Elements Data Storage query callback list."));
+}
\ No newline at end of file
diff --git a/Engine/Plugins/Experimental/TypedElementsDataStorage/Source/TypedElementsDataStorage/Private/TypedElementDatabase.h b/Engine/Plugins/Experimental/TypedElementsDataStorage/Source/TypedElementsDataStorage/Private/TypedElementDatabase.h
--- a/Engine/Plugins/Experimental/TypedElementsDataStorage/Source/TypedElementsDataStorage/Private/TypedElementDatabase.h
+++ b/Engine/Plugins/Experimental/TypedElementsDataStorage/Source/TypedElementsDataStorage/Private/TypedElementDatabase.h
@@ -15,6 +15,7 @@
 
 struct FMassEntityManager;
 struct FMassProcessingPhaseManager;
+class FOutputDevice;
 class UWorld;
 
 struct FTypedElementDatabaseExtendedQuery
@@ -79,6 +80,7 @@ public:
 	bool IsAvailable() const override;
 	void* GetExternalSystemAddress(UClass* Target) override;
 
+	void DebugPrintQueryCallbacks(FOutputDevice& Output);
 private:
 	using QueryStore = TTypedElementHandleStore<FTypedElementDatabaseExtendedQuery>;
 
diff --git a/Engine/Plugins/Experimental/TypedElementsDataStorage/Source/TypedElementsDataStorage/Private/TypedElementHandleStore.h b/Engine/Plugins/Experimental/TypedElementsDataStorage/Source/TypedElementsDataStorage/Private/TypedElementHandleStore.h
--- a/Engine/Plugins/Experimental/TypedElementsDataStorage/Source/TypedElementsDataStorage/Private/TypedElementHandleStore.h
+++ b/Engine/Plugins/Experimental/TypedElementsDataStorage/Source/TypedElementsDataStorage/Private/TypedElementHandleStore.h
@@ -2,21 +2,29 @@
 
 #pragma once
 
+#include <type_traits>
 #include "Containers/Array.h"
 #include "Containers/Deque.h"
+#include "Templates/Function.h"
 
 template<typename DataType, uint32 ReservationSize = 64>
 class TTypedElementHandleStore
 {
 public:
-	struct HandleData
+	struct FGeneration
+	{
+		uint32 Generation : 31;
+		uint32 bIsAlive : 1;
+	};
+	
+	struct FHandleData
 	{
-		uint32 Index;
 		uint32 Generation;
+		uint32 Index;
 	};
 	union Handle
 	{
-		HandleData Data;
+		FHandleData Data;
 		uint64 Handle;
 	};
 	
@@ -27,14 +35,15 @@ public:
 	DataType& GetMutable(Handle Entry);
 	const DataType& Get(Handle Entry) const;
 
-	/** Removes the entry at the given handle if alive. This doesn't destroy the entry until the slot is reused. */
+	/** Removes the entry at the given handle if alive. */
 	void Remove(Handle Entry);
 
 	bool IsAlive(Handle Entry) const;
+	void ListAliveEntries(const TFunctionRef<void(const DataType&)>& Callback) const;
 
 private:
 	TArray<DataType> Data;
-	TArray<uint32> Generations;
+	TArray<FGeneration> Generations;
 	TDeque<uint32> RecycleBin; // Using a deque to avoid the same slot being constantly reused.
 };
 
@@ -54,7 +63,7 @@ auto TTypedElementHandleStore<DataType, ReservationSize>::Emplace(Args... Argume
 		for (uint32 Index = Data.Num(); Index < NewSize; ++Index)
 		{
 			Data.AddZeroed();
-			Generations.Add(0);
+			Generations.AddZeroed();
 			RecycleBin.EmplaceLast(Index);
 		}
 	}
@@ -67,7 +76,8 @@ auto TTypedElementHandleStore<DataType, ReservationSize>::Emplace(Args... Argume
 
 	Handle Result;
 	Result.Data.Index = Index;
-	Result.Data.Generation = Generations[Index];
+	Result.Data.Generation = Generations[Index].Generation;
+	Generations[Index].bIsAlive = 1;
 	return Result;
 }
 
@@ -97,7 +107,13 @@ void TTypedElementHandleStore<DataType, ReservationSize>::Remove(Handle Entry)
 {
 	if (IsAlive(Entry))
 	{
-		++Generations[Entry.Data.Index];
+		FGeneration& Generation = Generations[Entry.Data.Index];
+		++Generation.Generation;
+		Generation.bIsAlive = 0;
+		if constexpr (std::is_destructible_v<DataType> && !std::is_trivially_destructible_v<DataType>)
+		{
+			Data[Entry.Data.Index].~DataType();
+		}
 		RecycleBin.EmplaceLast(Entry.Data.Index);
 	}
 }
@@ -107,5 +123,24 @@ bool TTypedElementHandleStore<DataType, ReservationSize>::IsAlive(Handle Entry)
 {
 	return
 		Entry.Data.Index < static_cast<uint32>(Generations.Num()) &&
-		Generations[Entry.Data.Index] == Entry.Data.Generation;
-}
\ No newline at end of file
+		Generations[Entry.Data.Index].Generation == Entry.Data.Generation;
+}
+
+template<typename DataType, uint32 ReservationSize>
+void TTypedElementHandleStore<DataType, ReservationSize>::ListAliveEntries(const TFunctionRef<void(const DataType&)>& Callback) const
+{
+	int32 Count = Data.Num();
+	const DataType* EntryIt = Data.GetData();
+	const FGeneration* GenerationIt = Generations.GetData();
+	
+	for (int32 Index = 0; Index < Count; ++Index)
+	{
+		if (GenerationIt->bIsAlive)
+		{
+			Callback(*EntryIt);
+		}
+
+		++EntryIt;
+		++GenerationIt;
+	}
+}
