Extended TEDS' ability to interact with Slate.

The following changes are added:
- The CDO factories have additional functions to register widget purposes and widget constructors.
- The CDO factories are now called in two phases, guaranteeing that registration of e.g. purposes happens before widget constructors are registered with them.
- The CDO factories can optionally set an order, where lower numbers mean that the run before factories with a higher number. This can allow orders to guaranteed further, e.g. when a widget constructor needs to overwrite an earlier registration.
- The registration of widget purposes has been expanded upon. Purposes are now explicitly registered and can be assigned a type. Based on the type the registration of a widget constructor behaves differently.
- Widget constructors can now optionally be registered with column associations.
- Support has been added to create an instance of a widget constructor. This allows for a constructor to be created, configured and then reused multiple times. This cuts down on the cost needed for creating constructors.
- Widget constructors can now be matched with a purpose and a set of columns. In the latter case the widget constructors associated with a purpose will be matched with the requested columns, prioritizing the longest chain of matching columns. This is repeated until all columns are matched or there no columns left to match.
- The remaining generic code to setup the counter widget in the status bar has been moved to the counter widget file as it's no possible to do the full registration there.
- The tooltip on the counter widget is working again.
- A new function and console command to retrieve the widget purposes including their documentation.
- Added an optimization that avoids creating temporary copy of widgets if there are no arguments going to be applied to them.

diff --git a/Engine/Plugins/Experimental/TypedElementsDataStorage/Source/TypedElementsDataStorage/Private/TypedElementDatabaseUI.cpp b/Engine/Plugins/Experimental/TypedElementsDataStorage/Source/TypedElementsDataStorage/Private/TypedElementDatabaseUI.cpp
--- a/Engine/Plugins/Experimental/TypedElementsDataStorage/Source/TypedElementsDataStorage/Private/TypedElementDatabaseUI.cpp
+++ b/Engine/Plugins/Experimental/TypedElementsDataStorage/Source/TypedElementsDataStorage/Private/TypedElementDatabaseUI.cpp
@@ -3,13 +3,33 @@
 #include "TypedElementDatabaseUI.h"
 
 #include "Elements/Columns/TypedElementSlateWidgetColumns.h"
+#include "Elements/Interfaces/TypedElementDataStorageInterface.h"
+#include "Elements/Interfaces/TypedElementDataStorageCompatibilityInterface.h"
 #include "GenericPlatform/GenericPlatformMemory.h"
+#include "Widgets/SlateControlledConstruction.h"
 
-void UTypedElementDatabaseUi::Initialize(ITypedElementDataStorageInterface* StorageInterface)
+DEFINE_LOG_CATEGORY(LogTypedElementDatabaseUI);
+
+namespace Internal
+{
+	// Source: https://en.cppreference.com/w/cpp/utility/variant/visit
+	template<class... Ts>
+	struct TOverloaded : Ts...
+	{ 
+		using Ts::operator()...; 
+	};
+
+	template<class... Ts> TOverloaded(Ts...) -> TOverloaded<Ts...>;
+}
+
+void UTypedElementDatabaseUi::Initialize(
+	ITypedElementDataStorageInterface* StorageInterface,
+	ITypedElementDataStorageCompatibilityInterface* StorageCompatibilityInterface)
 {
 	checkf(StorageInterface, TEXT("Typed Element's Database compatibility manager is being initialized with an invalid storage target."));
 
 	Storage = StorageInterface;
+	StorageCompatibility = StorageCompatibilityInterface;
 	CreateStandardArchetypes();
 }
 
@@ -17,41 +37,369 @@ void UTypedElementDatabaseUi::Deinitialize()
 {
 }
 
-void UTypedElementDatabaseUi::RegisterWidgetFactory(FName Purpose, const UScriptStruct* Constructor)
+void UTypedElementDatabaseUi::RegisterWidgetPurpose(FName Purpose, EPurposeType Type, FText Description)
+{
+	FPurposeInfo* PurposeInfo = WidgetPurposes.Find(Purpose);
+	if (!PurposeInfo)
+	{
+		FPurposeInfo& NewInfo = WidgetPurposes.Add(Purpose);
+		NewInfo.Type = Type;
+		NewInfo.Description = MoveTemp(Description);
+	}
+}
+
+bool UTypedElementDatabaseUi::RegisterWidgetFactory(FName Purpose, const UScriptStruct* Constructor)
 {
 	checkf(Constructor->IsChildOf(FTypedElementWidgetConstructor::StaticStruct()),
-		TEXT("Attempting to register a Typed Elements widget constructor '%s' that isn't deriving from FTypedElementWidgetConstructor."),
+		TEXT("Attempting to register a Typed Elements widget constructor '%s' that isn't derived from FTypedElementWidgetConstructor."),
 		*Constructor->GetFullName());
-	WidgetFactoryStructs.Add(Purpose, Constructor);
+	
+	if (FPurposeInfo* PurposeInfo = WidgetPurposes.Find(Purpose))
+	{
+		switch (PurposeInfo->Type)
+		{
+		case ITypedElementDataStorageUiInterface::EPurposeType::Generic:
+			PurposeInfo->Factories.Emplace(Constructor);
+			PurposeInfo->bIsSorted = false;
+			return true;
+		case ITypedElementDataStorageUiInterface::EPurposeType::UniqueByName:
+			if (PurposeInfo->Factories.IsEmpty())
+			{
+				PurposeInfo->Factories.Emplace(Constructor);
+				PurposeInfo->bIsSorted = false;
+			}
+			else
+			{
+				PurposeInfo->Factories.EmplaceAt(0, Constructor);
+			}
+			return true;
+		case ITypedElementDataStorageUiInterface::EPurposeType::UniqueByNameAndColumn:
+			UE_LOG(LogTypedElementDatabaseUI, Warning,
+				TEXT("Unable to register widget factory '%s' as purpose '%s' requires at least one column for matching."), 
+				*Constructor->GetName(), *Purpose.ToString());
+			return false;
+		default:
+			checkf(false, TEXT("Unexpected ITypedElementDataStorageUiInterface::EPurposeType found provided when registering widget factory."));
+			return false;
+		}
+	}
+	else
+	{
+		UE_LOG(LogTypedElementDatabaseUI, Warning, 
+			TEXT("Unable to register widget factory '%s' as purpose '%s' isn't registered."), *Constructor->GetName(), *Purpose.ToString());
+		return false;
+	}
 }
 
-void UTypedElementDatabaseUi::RegisterWidgetFactory(
-	FName Purpose, const UScriptStruct* ConstructorType, TUniquePtr<FTypedElementWidgetConstructor>&& Constructor)
+bool UTypedElementDatabaseUi::RegisterWidgetFactory(
+	FName Purpose, const UScriptStruct* Constructor, TArray<TWeakObjectPtr<const UScriptStruct>> Columns)
 {
-	WidgetFactoryInstances.Add(Purpose, { MoveTemp(Constructor), ConstructorType });
+	if (!Columns.IsEmpty())
+	{
+		checkf(Constructor->IsChildOf(FTypedElementWidgetConstructor::StaticStruct()),
+			TEXT("Attempting to register a Typed Elements widget constructor '%s' that isn't deriving from FTypedElementWidgetConstructor."),
+			*Constructor->GetFullName());
+
+		if (FPurposeInfo* PurposeInfo = WidgetPurposes.Find(Purpose))
+		{
+			switch (PurposeInfo->Type)
+			{
+			case ITypedElementDataStorageUiInterface::EPurposeType::Generic:
+				PurposeInfo->Factories.Emplace(Constructor);
+				PurposeInfo->bIsSorted = false;
+				return true;
+			case ITypedElementDataStorageUiInterface::EPurposeType::UniqueByName:
+				if (PurposeInfo->Factories.IsEmpty())
+				{
+					PurposeInfo->Factories.Emplace(Constructor);
+					PurposeInfo->bIsSorted = false;
+				}
+				else
+				{
+					PurposeInfo->Factories.EmplaceAt(0, Constructor);
+				}
+				return true;
+			case ITypedElementDataStorageUiInterface::EPurposeType::UniqueByNameAndColumn:
+				PrepareColumnsList(Columns);
+				PurposeInfo->Factories.Emplace(Constructor, MoveTemp(Columns));
+				PurposeInfo->bIsSorted = false;
+				return true;
+			default:
+				checkf(false, TEXT("Unexpected ITypedElementDataStorageUiInterface::EPurposeType found provided when registering widget factory."));
+				return false;
+			}
+		}
+		else
+		{
+			UE_LOG(LogTypedElementDatabaseUI, Warning,
+				TEXT("Unable to register widget factory '%s' as purpose '%s' isn't registered."), *Constructor->GetName(), *Purpose.ToString());
+			return false;
+		}
+	}
+	else
+	{
+		return RegisterWidgetFactory(Purpose, Constructor);
+	}
 }
 
-void UTypedElementDatabaseUi::ConstructWidgets(FName Purpose, TConstArrayView<TypedElement::ColumnUtils::Argument> Arguments,
-	const WidgetCreatedCallback& ConstructionCallback)
+bool UTypedElementDatabaseUi::RegisterWidgetFactory(FName Purpose, TUniquePtr<FTypedElementWidgetConstructor>&& Constructor)
 {
-	// Construct with structs
-	for (TMultiMap<FName, const UScriptStruct*>::TKeyIterator KeyValue = WidgetFactoryStructs.CreateKeyIterator(Purpose); 
-		KeyValue; ++KeyValue)
+	checkf(Constructor->GetTypeInfo(), TEXT("Widget constructor being registered that doesn't have valid type information."));
+	
+	if (FPurposeInfo* PurposeInfo = WidgetPurposes.Find(Purpose))
 	{
-		const UScriptStruct* Target = KeyValue.Value();
-		CreateWidgetInstanceFromDescription(Target, Arguments, ConstructionCallback);
+		switch (PurposeInfo->Type)
+		{
+		case ITypedElementDataStorageUiInterface::EPurposeType::Generic:
+			PurposeInfo->Factories.Emplace(MoveTemp(Constructor));
+			PurposeInfo->bIsSorted = false;
+			return true;
+		case ITypedElementDataStorageUiInterface::EPurposeType::UniqueByName:
+			if (PurposeInfo->Factories.IsEmpty())
+			{
+				PurposeInfo->Factories.Emplace(MoveTemp(Constructor));
+				PurposeInfo->bIsSorted = false;
+			}
+			else
+			{
+				PurposeInfo->Factories.EmplaceAt(0, MoveTemp(Constructor));
+			}
+			return true;
+		case ITypedElementDataStorageUiInterface::EPurposeType::UniqueByNameAndColumn:
+			UE_LOG(LogTypedElementDatabaseUI, Warning,
+				TEXT("Unable to register widget factory '%s' as purpose '%s' requires at least one column for matching."),
+				*Constructor->GetTypeInfo()->GetName(), *Purpose.ToString());
+			return false;
+		default:
+			checkf(false, TEXT("Unexpected ITypedElementDataStorageUiInterface::EPurposeType found provided when registering widget factory."));
+			return false;
+		}
+	}
+	else
+	{
+		UE_LOG(LogTypedElementDatabaseUI, Warning, 
+			TEXT("Unable to register widget factory as purpose '%s' isn't registered."), *Purpose.ToString());
+		return false;
+	}
+}
+
+bool UTypedElementDatabaseUi::RegisterWidgetFactory(FName Purpose, TUniquePtr<FTypedElementWidgetConstructor>&& Constructor, 
+	TArray<TWeakObjectPtr<const UScriptStruct>> Columns)
+{
+	if (!Columns.IsEmpty())
+	{
+		checkf(Constructor->GetTypeInfo(), TEXT("Widget constructor being registered that doesn't have valid type information."));
+
+		if (FPurposeInfo* PurposeInfo = WidgetPurposes.Find(Purpose))
+		{
+			switch (PurposeInfo->Type)
+			{
+			case ITypedElementDataStorageUiInterface::EPurposeType::Generic:
+				PurposeInfo->Factories.Emplace(MoveTemp(Constructor));
+				PurposeInfo->bIsSorted = false;
+				return true;
+			case ITypedElementDataStorageUiInterface::EPurposeType::UniqueByName:
+				if (PurposeInfo->Factories.IsEmpty())
+				{
+					PurposeInfo->Factories.Emplace(MoveTemp(Constructor));
+					PurposeInfo->bIsSorted = false;
+				}
+				else
+				{
+					PurposeInfo->Factories.EmplaceAt(0, MoveTemp(Constructor));
+				}
+				return true;
+			case ITypedElementDataStorageUiInterface::EPurposeType::UniqueByNameAndColumn:
+				PrepareColumnsList(Columns);
+				PurposeInfo->Factories.Emplace(MoveTemp(Constructor), MoveTemp(Columns));
+				PurposeInfo->bIsSorted = false;
+				return true;
+			default:
+				checkf(false, TEXT("Unexpected ITypedElementDataStorageUiInterface::EPurposeType found provided when registering widget factory."));
+				return false;
+			}
+		}
+		else
+		{
+			UE_LOG(LogTypedElementDatabaseUI, Warning, TEXT("Unable to register widget factory '%s' as purpose '%s' isn't registered."), 
+				*Constructor->GetTypeInfo()->GetName(), *Purpose.ToString());
+			return false;
+		}
+	}
+	else
+	{
+		return RegisterWidgetFactory(Purpose, MoveTemp(Constructor));
+	}
+}
+
+void UTypedElementDatabaseUi::CreateWidgetConstructors(FName Purpose,
+	TConstArrayView<TypedElement::ColumnUtils::Argument> Arguments, const WidgetConstructorCallback& Callback)
+{
+	if (FPurposeInfo* PurposeInfo = WidgetPurposes.Find(Purpose))
+	{
+		for (const FWidgetFactory& Factory : PurposeInfo->Factories)
+		{
+			if (!CreateSingleWidgetConstructor(Factory.Constructor, Arguments, {}, Callback))
+			{
+				return;
+			}
+		}
+	}
+}
+
+void UTypedElementDatabaseUi::CreateWidgetConstructors(FName Purpose, TArray<TWeakObjectPtr<const UScriptStruct>>& Columns,
+	TConstArrayView<TypedElement::ColumnUtils::Argument> Arguments, const WidgetConstructorCallback& Callback)
+{
+	using ColumnView = TConstArrayView<TWeakObjectPtr<const UScriptStruct>>;
+	
+	if (FPurposeInfo* PurposeInfo = WidgetPurposes.Find(Purpose))
+	{
+		// Sort so searching can be done in a single pass. This would also allow for binary searching, but the number of columns
+		// is typically small enough for a binary search to end up being more expensive than a linear search. This may change
+		// if/when there are a sufficient enough number of widgets that are bound to a large number of columns.
+		Columns.Sort(
+			[](const TWeakObjectPtr<const UScriptStruct>& Lhs, const TWeakObjectPtr<const UScriptStruct>& Rhs)
+			{
+				return Lhs.GetWeakPtrTypeHash() < Rhs.GetWeakPtrTypeHash();
+			});
+
+		if (!PurposeInfo->bIsSorted)
+		{
+			// This is the only call that requires the array of factories to be sorted from from largest to smallest number
+			// of columns, so lazily sort only when needed.
+			PurposeInfo->Factories.StableSort(
+				[](const FWidgetFactory& Lhs, const FWidgetFactory& Rhs)
+				{
+					int32 LeftSize = Lhs.Columns.Num();
+					int32 RightSize = Rhs.Columns.Num();
+					if (LeftSize == RightSize)
+					{
+						for (int32 Index = 0; Index < LeftSize; ++Index)
+						{
+							uint32 LeftHash = Lhs.Columns[Index].GetWeakPtrTypeHash();
+							uint32 RightHash = Lhs.Columns[Index].GetWeakPtrTypeHash();
+							if (LeftHash != RightHash)
+							{
+								return LeftHash < RightHash;
+							}
+						}
+						return true;
+					}
+					else
+					{
+						return LeftSize > RightSize;
+					}
+				});
+			PurposeInfo->bIsSorted = true;
+		}
+
+		for (auto FactoryIt = PurposeInfo->Factories.CreateConstIterator();  FactoryIt && !Columns.IsEmpty();  ++FactoryIt)
+		{
+			int32 MatchingIndex = INDEX_NONE;
+			if (FactoryIt->Columns.Num() > Columns.Num())
+			{
+				// There are more columns required for this factory than there are in the requested columns list so skip this
+				// factory.
+				continue;
+			}
+			
+			ColumnView ColumnsRange = Columns;
+			for (auto ColumnIt = FactoryIt->Columns.CreateConstIterator(); ColumnIt && !ColumnsRange.IsEmpty(); ++ColumnIt)
+			{
+				int32 Index = INDEX_NONE;
+				if (ColumnsRange.Find(*ColumnIt, Index))
+				{
+					MatchingIndex = MatchingIndex == INDEX_NONE ? Index : MatchingIndex;
+					ColumnsRange = ColumnsRange.RightChop(Index + 1);
+				}
+				else
+				{
+					// Reset as there is at least one column that didn't match.
+					MatchingIndex = INDEX_NONE;
+					break;
+				}
+			}
+
+			// This assumes that the column arrays are sorted from longest to shortest.
+			if (MatchingIndex != INDEX_NONE)
+			{
+				// Delete in reverse order to reduce the number of objects that need to be moved up as the order needs to remain stable.
+				ColumnsRange = ColumnView(Columns);
+				for (auto It = FactoryIt->Columns.rbegin(); It != FactoryIt->Columns.rend(); ++It)
+				{
+					int32 Index = INDEX_NONE;
+					ensureMsgf(ColumnsRange.Find(*It, Index), TEXT("A previous found matching column can't be found in the original array."));
+					Columns.RemoveAt(Index);
+					ColumnsRange = ColumnsRange.Left(Index);
+				}
+
+				if (!CreateSingleWidgetConstructor(FactoryIt->Constructor, Arguments, FactoryIt->Columns, Callback))
+				{
+					return;
+				}
+			}
+		}
 	}
+}
 
-	// Construct with instances
-	for (TMultiMap<FName, FInstanceConstructor>::TKeyIterator KeyValue = WidgetFactoryInstances.CreateKeyIterator(Purpose); 
-		KeyValue; ++KeyValue)
+void UTypedElementDatabaseUi::ConstructWidgets(FName Purpose, TConstArrayView<TypedElement::ColumnUtils::Argument> Arguments,
+	const WidgetCreatedCallback& ConstructionCallback)
+{
+	if (FPurposeInfo* PurposeInfo = WidgetPurposes.Find(Purpose))
 	{
-		FTypedElementWidgetConstructor* Source = KeyValue.Value().Constructor.Get();
-		const UScriptStruct* Target = KeyValue.Value().ConstructorType;
-		CreateWidgetInstanceFromInstance(Source, Target, Arguments, ConstructionCallback);
+		for (const FWidgetFactory& Factory : PurposeInfo->Factories)
+		{
+			std::visit(Internal::TOverloaded
+				{
+					[this, &Arguments, &ConstructionCallback](const UScriptStruct* Constructor)
+					{ 
+						CreateWidgetInstanceFromDescription(Constructor, Arguments, ConstructionCallback); 
+					},
+					[this, &Arguments, &ConstructionCallback](const TUniquePtr<FTypedElementWidgetConstructor>& Constructor)
+					{
+						CreateWidgetInstanceFromInstance(Constructor.Get(), Arguments, ConstructionCallback);
+					}
+				}, Factory.Constructor);
+		}
 	}
 }
 
+bool UTypedElementDatabaseUi::CreateSingleWidgetConstructor(
+	const FWidgetFactory::ConstructorType& Constructor,
+	TConstArrayView<TypedElement::ColumnUtils::Argument> Arguments, 
+	TConstArrayView<TWeakObjectPtr<const UScriptStruct>> MatchedColumnTypes, 
+	const WidgetConstructorCallback& Callback)
+{
+	return std::visit(Internal::TOverloaded
+		{
+			[this, &Arguments, &MatchedColumnTypes, &Callback](const UScriptStruct* Target)
+			{
+				TUniquePtr<FTypedElementWidgetConstructor> Result(reinterpret_cast<FTypedElementWidgetConstructor*>(
+					FMemory::Malloc(Target->GetStructureSize(), Target->GetMinAlignment())));
+				if (Result)
+				{
+					Target->InitializeStruct(Result.Get());
+					return Callback(MoveTemp(Result), MatchedColumnTypes);
+				}
+				return true;
+			},
+			[this, &Arguments, &MatchedColumnTypes, &Callback](const TUniquePtr<FTypedElementWidgetConstructor>& Target)
+			{
+				const UScriptStruct* TargetType = Target->GetTypeInfo();
+				checkf(TargetType, TEXT("Expected valid type information from a widget constructor."));
+				TUniquePtr<FTypedElementWidgetConstructor> Result(reinterpret_cast<FTypedElementWidgetConstructor*>(
+					FMemory::Malloc(TargetType->GetStructureSize(), TargetType->GetMinAlignment())));
+				if (Result)
+				{
+					TargetType->InitializeStruct(Result.Get());
+					TargetType->CopyScriptStruct(Result.Get(), Target.Get());
+					return Callback(MoveTemp(Result), MatchedColumnTypes);
+				}
+				return true;
+			}
+		}, Constructor);
+}
+
 void UTypedElementDatabaseUi::CreateWidgetInstanceFromDescription(
 	const UScriptStruct* Target,
 	TConstArrayView<TypedElement::ColumnUtils::Argument> Arguments,
@@ -62,7 +410,8 @@ void UTypedElementDatabaseUi::CreateWidgetInstanceFromDescription(
 	if (Constructor)
 	{
 		Target->InitializeStruct(Constructor);
-		CreateWidgetInstance(*Constructor, Arguments, ConstructionCallback, Target);
+		CreateWidgetInstance(*Constructor, Arguments, ConstructionCallback);
+		Target->DestroyStruct(&Constructor);
 	}
 	else
 	{
@@ -72,29 +421,37 @@ void UTypedElementDatabaseUi::CreateWidgetInstanceFromDescription(
 
 void UTypedElementDatabaseUi::CreateWidgetInstanceFromInstance(
 	FTypedElementWidgetConstructor* SourceConstructor,
-	const UScriptStruct* Target,
 	TConstArrayView<TypedElement::ColumnUtils::Argument> Arguments,
 	const WidgetCreatedCallback& ConstructionCallback)
 {
-	FTypedElementWidgetConstructor* Constructor = reinterpret_cast<FTypedElementWidgetConstructor*>(
-		FMemory_Alloca_Aligned(Target->GetStructureSize(), Target->GetMinAlignment()));
-	if (Constructor)
+	if (Arguments.IsEmpty())
 	{
-		Target->InitializeStruct(Constructor);
-		Target->CopyScriptStruct(Constructor, SourceConstructor);
-		CreateWidgetInstance(*Constructor, Arguments, ConstructionCallback, Target);
+		CreateWidgetInstance(*SourceConstructor, Arguments, ConstructionCallback);
 	}
 	else
 	{
-		checkf(false, TEXT("Remaining stack space is too small to create a Typed Elements widget constructor from an instance."));
+		const UScriptStruct* Target = SourceConstructor->GetTypeInfo();
+		checkf(Target, TEXT("Expected valid type information from a widget constructor."));
+		FTypedElementWidgetConstructor* Constructor = reinterpret_cast<FTypedElementWidgetConstructor*>(
+			FMemory_Alloca_Aligned(Target->GetStructureSize(), Target->GetMinAlignment()));
+		if (Constructor)
+		{
+			Target->InitializeStruct(Constructor);
+			Target->CopyScriptStruct(Constructor, SourceConstructor);
+			CreateWidgetInstance(*Constructor, Arguments, ConstructionCallback);
+			Target->DestroyStruct(&Constructor);
+		}
+		else
+		{
+			checkf(false, TEXT("Remaining stack space is too small to create a Typed Elements widget constructor from an instance."));
+		}
 	}
 }
 
 void UTypedElementDatabaseUi::CreateWidgetInstance(
 	FTypedElementWidgetConstructor& Constructor, 
 	TConstArrayView<TypedElement::ColumnUtils::Argument> Arguments, 
-	const WidgetCreatedCallback& ConstructionCallback, 
-	const UScriptStruct* Target)
+	const WidgetCreatedCallback& ConstructionCallback)
 {
 	TypedElementRowHandle Row = Storage->AddRow(WidgetTable);
 	TSharedPtr<SWidget> Widget = Constructor.Construct(Row, Storage, this, Arguments);
@@ -106,14 +463,43 @@ void UTypedElementDatabaseUi::CreateWidgetInstance(
 	{
 		Storage->RemoveRow(Row);
 	}
-
-	Target->DestroyStruct(&Constructor);
 }
 
 TSharedPtr<SWidget> UTypedElementDatabaseUi::ConstructWidget(TypedElementRowHandle Row, FTypedElementWidgetConstructor& Constructor,
 	TConstArrayView<TypedElement::ColumnUtils::Argument> Arguments)
 {
-	return Constructor.Construct(Row, Storage, this, Arguments);
+	if (Arguments.IsEmpty())
+	{
+		return Constructor.Construct(Row, Storage, this, Arguments);
+	}
+	else
+	{
+		const UScriptStruct* Target = Constructor.GetTypeInfo();
+		FTypedElementWidgetConstructor* ConstructorCopy = reinterpret_cast<FTypedElementWidgetConstructor*>(
+			FMemory_Alloca_Aligned(Target->GetStructureSize(), Target->GetMinAlignment()));
+		if (ConstructorCopy)
+		{
+			Target->InitializeStruct(ConstructorCopy);
+			Target->CopyScriptStruct(ConstructorCopy, &Constructor);
+			TSharedPtr<SWidget> Widget = ConstructorCopy->Construct(Row, Storage, this, Arguments);
+			Target->DestroyStruct(&ConstructorCopy);
+			return Widget;
+		}
+		else
+		{
+			checkf(false, TEXT("Remaining stack space is too small to create a Typed Elements widget constructor from an instance."));
+			return nullptr;
+		}
+		return nullptr;
+	}
+}
+
+void UTypedElementDatabaseUi::ListWidgetPurposes(const WidgetPurposeCallback& Callback) const
+{
+	for (auto&& It : WidgetPurposes)
+	{
+		Callback(It.Key, It.Value.Type, It.Value.Description);
+	}
 }
 
 void UTypedElementDatabaseUi::CreateStandardArchetypes()
@@ -123,4 +509,60 @@ void UTypedElementDatabaseUi::CreateStandardArchetypes()
 			FTypedElementSlateWidgetReferenceColumn::StaticStruct(),
 			FTypedElementSlateWidgetReferenceDeletesRowTag::StaticStruct()
 		}), FName("Editor_WidgetTable"));
-}
\ No newline at end of file
+}
+
+void UTypedElementDatabaseUi::PrepareColumnsList(TArray<TWeakObjectPtr<const UScriptStruct>>& Columns)
+{
+	Columns.Sort(
+		[](const TWeakObjectPtr<const UScriptStruct>& Lhs, const TWeakObjectPtr<const UScriptStruct>& Rhs)
+		{
+			return Lhs.GetWeakPtrTypeHash() < Rhs.GetWeakPtrTypeHash();
+		});
+
+	// Remove duplicates.
+	for (int32 Index = 1; Index < Columns.Num(); )
+	{
+		if (Columns[Index - 1].GetWeakPtrTypeHash() != Columns[Index].GetWeakPtrTypeHash())
+		{
+			++Index;
+		}
+		else
+		{
+			Columns.RemoveAt(Index);
+		}
+	}
+}
+
+
+
+//
+// FWidgetFactory
+//
+
+UTypedElementDatabaseUi::FWidgetFactory::FWidgetFactory(const UScriptStruct* InConstructor)
+	: Constructor(InConstructor)
+{
+}
+
+UTypedElementDatabaseUi::FWidgetFactory::FWidgetFactory(TUniquePtr<FTypedElementWidgetConstructor>&& InConstructor)
+	: Constructor(MoveTemp(InConstructor))
+{
+	checkf(std::get<TUniquePtr<FTypedElementWidgetConstructor>>(Constructor)->GetTypeInfo(), 
+		TEXT("Widget constructor registered that didn't contain valid type information."));
+}
+
+UTypedElementDatabaseUi::FWidgetFactory::FWidgetFactory(const UScriptStruct* InConstructor, 
+	TArray<TWeakObjectPtr<const UScriptStruct>>&& InColumns)
+	: Columns(MoveTemp(InColumns))
+	, Constructor(InConstructor)
+{
+}
+
+UTypedElementDatabaseUi::FWidgetFactory::FWidgetFactory(TUniquePtr<FTypedElementWidgetConstructor>&& InConstructor, 
+	TArray<TWeakObjectPtr<const UScriptStruct>>&& InColumns)
+	: Columns(MoveTemp(InColumns))
+	, Constructor(MoveTemp(InConstructor))
+{
+	checkf(std::get<TUniquePtr<FTypedElementWidgetConstructor>>(Constructor)->GetTypeInfo(),
+		TEXT("Widget constructor registered that didn't contain valid type information."));
+}
diff --git a/Engine/Plugins/Experimental/TypedElementsDataStorage/Source/TypedElementsDataStorage/Private/TypedElementDatabaseUI.h b/Engine/Plugins/Experimental/TypedElementsDataStorage/Source/TypedElementsDataStorage/Private/TypedElementDatabaseUI.h
--- a/Engine/Plugins/Experimental/TypedElementsDataStorage/Source/TypedElementsDataStorage/Private/TypedElementDatabaseUI.h
+++ b/Engine/Plugins/Experimental/TypedElementsDataStorage/Source/TypedElementsDataStorage/Private/TypedElementDatabaseUI.h
@@ -2,13 +2,19 @@
 
 #pragma once
 
+#include <variant>
 #include "Containers/Map.h"
-#include "Elements/Interfaces/TypedElementDataStorageInterface.h"
 #include "Elements/Interfaces/TypedElementDataStorageUiInterface.h"
+#include "Logging/LogMacros.h"
 #include "UObject/ObjectMacros.h"
 
 #include "TypedElementDatabaseUI.generated.h"
 
+class ITypedElementDataStorageInterface;
+class ITypedElementDataStorageCompatibilityInterface;
+
+TYPEDELEMENTSDATASTORAGE_API DECLARE_LOG_CATEGORY_EXTERN(LogTypedElementDatabaseUI, Log, All);
+
 UCLASS()
 class TYPEDELEMENTSDATASTORAGE_API UTypedElementDatabaseUi final
 	: public UObject
@@ -19,29 +25,63 @@ class TYPEDELEMENTSDATASTORAGE_API UTypedElementDatabaseUi final
 public:
 	~UTypedElementDatabaseUi() override = default;
 
-	void Initialize(ITypedElementDataStorageInterface* StorageInterface);
+	void Initialize(
+		ITypedElementDataStorageInterface* StorageInterface, 
+		ITypedElementDataStorageCompatibilityInterface* StorageCompatibilityInterface);
 	void Deinitialize();
 
-	void RegisterWidgetFactory(FName Purpose, const UScriptStruct* Constructor) override;
-	
+	void RegisterWidgetPurpose(FName Purpose, EPurposeType Type, FText Description) override;
+
+	bool RegisterWidgetFactory(FName Purpose, const UScriptStruct* Constructor) override;
+	bool RegisterWidgetFactory(FName Purpose, const UScriptStruct* Constructor,
+		TArray<TWeakObjectPtr<const UScriptStruct>> Columns) override;
+	bool RegisterWidgetFactory(FName Purpose, TUniquePtr<FTypedElementWidgetConstructor>&& Constructor) override;
+	bool RegisterWidgetFactory(FName Purpose, TUniquePtr<FTypedElementWidgetConstructor>&& Constructor,
+		TArray<TWeakObjectPtr<const UScriptStruct>> Columns) override;
+
+	void CreateWidgetConstructors(FName Purpose,
+		TConstArrayView<TypedElement::ColumnUtils::Argument> Arguments, const WidgetConstructorCallback& Callback) override;
+	void CreateWidgetConstructors(FName Purpose, TArray<TWeakObjectPtr<const UScriptStruct>>& Columns,
+		TConstArrayView<TypedElement::ColumnUtils::Argument> Arguments, const WidgetConstructorCallback& Callback) override;
+
 	void ConstructWidgets(FName Purpose, TConstArrayView<TypedElement::ColumnUtils::Argument> Arguments,
 		const WidgetCreatedCallback& ConstructionCallback) override;
 	TSharedPtr<SWidget> ConstructWidget(TypedElementRowHandle Row, FTypedElementWidgetConstructor& Constructor,
 		TConstArrayView<TypedElement::ColumnUtils::Argument> Arguments) override;
 
-protected:
-	void RegisterWidgetFactory(FName Purpose, const UScriptStruct* ConstructorType,
-		TUniquePtr<FTypedElementWidgetConstructor>&& Constructor) override;
+	void ListWidgetPurposes(const WidgetPurposeCallback& Callback) const override;
 
 private:
-	struct FInstanceConstructor
+	struct FWidgetFactory
 	{
-		TUniquePtr<FTypedElementWidgetConstructor> Constructor;
-		const UScriptStruct* ConstructorType;
+		using ConstructorType = std::variant<const UScriptStruct*, TUniquePtr<FTypedElementWidgetConstructor>>;
+
+		TArray<TWeakObjectPtr<const UScriptStruct>> Columns;
+		ConstructorType Constructor;
+
+		FWidgetFactory() = default;
+		explicit FWidgetFactory(const UScriptStruct* InConstructor);
+		explicit FWidgetFactory(TUniquePtr<FTypedElementWidgetConstructor>&& InConstructor);
+		FWidgetFactory(const UScriptStruct* InConstructor, TArray<TWeakObjectPtr<const UScriptStruct>>&& InColumns);
+		FWidgetFactory(TUniquePtr<FTypedElementWidgetConstructor>&& InConstructor, TArray<TWeakObjectPtr<const UScriptStruct>>&& InColumns);
+	};
+
+	struct FPurposeInfo
+	{	
+		TArray<FWidgetFactory> Factories;
+		FText Description;
+		EPurposeType Type;
+		bool bIsSorted{ false }; //< Whether or not the array of factories needs to be sorted. The factories themselves are already sorted.
 	};
 
 	void CreateStandardArchetypes();
 
+	bool CreateSingleWidgetConstructor(
+		const FWidgetFactory::ConstructorType& Constructor,
+		TConstArrayView<TypedElement::ColumnUtils::Argument> Arguments, 
+		TConstArrayView<TWeakObjectPtr<const UScriptStruct>> MatchedColumnTypes,
+		const WidgetConstructorCallback& Callback);
+
 	void CreateWidgetInstanceFromDescription(
 		const UScriptStruct* Target,
 		TConstArrayView<TypedElement::ColumnUtils::Argument> Arguments,
@@ -49,18 +89,20 @@ private:
 
 	void CreateWidgetInstanceFromInstance(
 		FTypedElementWidgetConstructor* SourceConstructor,
-		const UScriptStruct* Target,
 		TConstArrayView<TypedElement::ColumnUtils::Argument> Arguments,
 		const WidgetCreatedCallback& ConstructionCallback);
 
 	void CreateWidgetInstance(
 		FTypedElementWidgetConstructor& Constructor,
 		TConstArrayView<TypedElement::ColumnUtils::Argument> Arguments,
-		const WidgetCreatedCallback& ConstructionCallback,
-		const UScriptStruct* Target);
+		const WidgetCreatedCallback& ConstructionCallback);
+
+	static void PrepareColumnsList(TArray<TWeakObjectPtr<const UScriptStruct>>& Columns);
 
-	ITypedElementDataStorageInterface* Storage{ nullptr };
 	TypedElementTableHandle WidgetTable{ TypedElementInvalidTableHandle };
-	TMultiMap<FName, const UScriptStruct*> WidgetFactoryStructs;
-	TMultiMap<FName, FInstanceConstructor> WidgetFactoryInstances;
+	
+	TMap<FName, FPurposeInfo> WidgetPurposes;
+	
+	ITypedElementDataStorageInterface* Storage{ nullptr };
+	ITypedElementDataStorageCompatibilityInterface* StorageCompatibility{ nullptr };
 };
\ No newline at end of file
diff --git a/Engine/Plugins/Experimental/TypedElementsDataStorage/Source/TypedElementsDataStorage/Private/TypedElementsDataStorage.cpp b/Engine/Plugins/Experimental/TypedElementsDataStorage/Source/TypedElementsDataStorage/Private/TypedElementsDataStorage.cpp
--- a/Engine/Plugins/Experimental/TypedElementsDataStorage/Source/TypedElementsDataStorage/Private/TypedElementsDataStorage.cpp
+++ b/Engine/Plugins/Experimental/TypedElementsDataStorage/Source/TypedElementsDataStorage/Private/TypedElementsDataStorage.cpp
@@ -14,9 +14,9 @@
 
 #define LOCTEXT_NAMESPACE "FTypedElementsDataStorageModule"
 
-// MASS uses DCO in a few places, making it a difficult to consistently register Type Element's Columns and Tags
+// MASS uses CDO in a few places, making it a difficult to consistently register Type Element's Columns and Tags
 // as they may have not been set up to impersonate MASS' Fragments and Tags yet. To get around this, do the 
-// registration during static initialization, which is before DCOs.
+// registration during static initialization, which is before CDOs.
 /**
  * Typed Elements provides base classes for columns and tags. These directly map to fragments and tags in MASS.
  * To avoid deep and tight coupling between both systems, columns and tags don't directly inhered from MASS, but
@@ -62,7 +62,7 @@ void FTypedElementsDataStorageModule::StartupModule()
 				DatabaseCompatibility->Initialize(Database.Get());
 
 				DatabaseUi = NewObject<UTypedElementDatabaseUi>();
-				DatabaseUi->Initialize(Database.Get());
+				DatabaseUi->Initialize(Database.Get(), DatabaseCompatibility.Get());
 
 				// Register the various database instances.
 				UTypedElementRegistry* Registry = UTypedElementRegistry::GetInstance();
@@ -75,16 +75,35 @@ void FTypedElementsDataStorageModule::StartupModule()
 				// Allow any factories to register their content.
 				TArray<UClass*> FactoryClasses;
 				GetDerivedClasses(UTypedElementDataStorageFactory::StaticClass(), FactoryClasses);
+				
+				TArray<UTypedElementDataStorageFactory*> Factories;
+				Factories.Reserve(FactoryClasses.Num());
 				for (UClass* Factory : FactoryClasses)
 				{
 					if (Factory->HasAnyClassFlags(CLASS_Abstract))
 					{
 						continue;
 					}
+					Factories.Add(GetMutableDefault<UTypedElementDataStorageFactory>(Factory));
+				}
+				Factories.StableSort(
+					[](const UTypedElementDataStorageFactory& Lhs, const UTypedElementDataStorageFactory& Rhs)
+					{
+						return Lhs.GetOrder() < Rhs.GetOrder();
+					});
 
-					UTypedElementDataStorageFactory* FactoryCDO = GetMutableDefault<UTypedElementDataStorageFactory>(Factory);
-					FactoryCDO->RegisterQueries(*Database);
-					FactoryCDO->RegisterWidgetConstructor(*Database, *DatabaseUi);
+				// First pass to call all registration without dependencies.
+				for (UTypedElementDataStorageFactory* Factory : Factories)
+				{
+					Factory->RegisterTables(*Database);
+					Factory->RegisterTickGroups(*Database);
+					Factory->RegisterWidgetPurposes(*DatabaseUi);
+				}
+				// Second pass to call all registration that would benefit or need the registration in the previous pass.
+				for (UTypedElementDataStorageFactory* Factory : Factories)
+				{
+					Factory->RegisterQueries(*Database);
+					Factory->RegisterWidgetConstructors(*Database, *DatabaseUi);
 				}
 
 				bInitialized = true;
diff --git a/Engine/Plugins/Experimental/TypedElementsDataStorage/Source/TypedElementsDataStorage/TypedElementsDataStorage.Build.cs b/Engine/Plugins/Experimental/TypedElementsDataStorage/Source/TypedElementsDataStorage/TypedElementsDataStorage.Build.cs
--- a/Engine/Plugins/Experimental/TypedElementsDataStorage/Source/TypedElementsDataStorage/TypedElementsDataStorage.Build.cs
+++ b/Engine/Plugins/Experimental/TypedElementsDataStorage/Source/TypedElementsDataStorage/TypedElementsDataStorage.Build.cs
@@ -19,15 +19,16 @@ public class TypedElementsDataStorage : ModuleRules
 					"Core",
 					"CoreUObject",
 					"Engine",
+					"EditorSubsystem",
 					"MassActors",
 					"MassEntity",
 					"MassEntityEditor",
 					"MassGameplayEditor",
 					"MassSimulation",
 					"TypedElementFramework",
-					"EditorSubsystem",
-					"UnrealEd",
-					"StructUtils"
+					"SlateCore",
+					"StructUtils",
+					"UnrealEd"
 				});
 
 			PrivateDependencyModuleNames.AddRange(new string[] {});
diff --git a/Engine/Plugins/Experimental/TypedElementsDataStorage/Source/TypedElementsDataStorageUI/Private/TypedElementCounterWidget.cpp b/Engine/Plugins/Experimental/TypedElementsDataStorage/Source/TypedElementsDataStorageUI/Private/TypedElementCounterWidget.cpp
--- a/Engine/Plugins/Experimental/TypedElementsDataStorage/Source/TypedElementsDataStorageUI/Private/TypedElementCounterWidget.cpp
+++ b/Engine/Plugins/Experimental/TypedElementsDataStorage/Source/TypedElementsDataStorageUI/Private/TypedElementCounterWidget.cpp
@@ -5,11 +5,17 @@
 #include "Elements/Columns/TypedElementSlateWidgetColumns.h"
 #include "Elements/Columns/TypedElementValueCacheColumns.h"
 #include "Elements/Framework/TypedElementQueryBuilder.h"
+#include "Elements/Framework/TypedElementRegistry.h"
 #include "Framework/Text/TextLayout.h"
+#include "Interfaces/IMainFrameModule.h"
 #include "Layout/Margin.h"
+#include "Misc/CoreDelegates.h"
+#include "ToolMenu.h"
+#include "ToolMenus.h"
+#include "TypedElementSubsystems.h"
 #include "Widgets/DeclarativeSyntaxSupport.h"
+#include "Widgets/SWindow.h"
 #include "Widgets/Text/STextBlock.h"
-#include "TypedElementSubsystems.h"
 
 #define LOCTEXT_NAMESPACE "TypedElementUI_CounterWidget"
 
@@ -18,13 +24,21 @@
 // UTypedElementCounterWidgetFactory
 //
 
+FName UTypedElementCounterWidgetFactory::WigetPurpose(TEXT("LevelEditor.StatusBar.ToolBar"));
+
+UTypedElementCounterWidgetFactory::UTypedElementCounterWidgetFactory()
+{
+	IMainFrameModule::Get().OnMainFrameCreationFinished().AddStatic(&UTypedElementCounterWidgetFactory::SetupMainWindowIntegrations);
+}
+
 void UTypedElementCounterWidgetFactory::RegisterQueries(ITypedElementDataStorageInterface& DataStorage) const
 {
 	using namespace TypedElementQueryBuilder;
 	using DSI = ITypedElementDataStorageInterface;
 
-	DataStorage.RegisterQuery(Select(TEXT("Sync Counter Widgets"), 
-		FProcessor(DSI::EQueryTickPhase::FrameEnd, DataStorage.GetQueryTickGroupName(DSI::EQueryTickGroups::SyncWidgets)).ForceToGameThread(true),
+	DataStorage.RegisterQuery(Select(TEXT("Sync counter widgets"), 
+		FProcessor(DSI::EQueryTickPhase::FrameEnd, DataStorage.GetQueryTickGroupName(DSI::EQueryTickGroups::SyncWidgets))
+			.ForceToGameThread(true),
 		[](
 			FCachedQueryContext<UTypedElementDataStorageSubsystem>& Context, 
 			FTypedElementSlateWidgetReferenceColumn& Widget,
@@ -57,13 +71,17 @@ void UTypedElementCounterWidgetFactory::RegisterQueries(ITypedElementDataStorage
 	).Compile());
 }
 
-void UTypedElementCounterWidgetFactory::RegisterWidgetConstructor(ITypedElementDataStorageInterface& DataStorage,
+void UTypedElementCounterWidgetFactory::RegisterWidgetPurposes(ITypedElementDataStorageUiInterface& DataStorageUi) const
+{
+	DataStorageUi.RegisterWidgetPurpose(WigetPurpose, ITypedElementDataStorageUiInterface::EPurposeType::Generic,
+		LOCTEXT("ToolBarPurposeDescription", "Widgets added to the status bar at the bottom editor of the main editor window."));
+}
+
+void UTypedElementCounterWidgetFactory::RegisterWidgetConstructors(ITypedElementDataStorageInterface& DataStorage,
 	ITypedElementDataStorageUiInterface& DataStorageUi) const
 {
 	using namespace TypedElementQueryBuilder;
 
-	FName Purpose("LevelEditor.StatusBar.ToolBar");
-
 	TUniquePtr<FTypedElementCounterWidgetConstructor> ActorCounter = MakeUnique<FTypedElementCounterWidgetConstructor>();
 	ActorCounter->LabelText = LOCTEXT("ActorCounterStatusBarLabel", "{0} {0}|plural(one=Actor, other=Actors)");
 	ActorCounter->ToolTipText = LOCTEXT(
@@ -72,9 +90,9 @@ void UTypedElementCounterWidgetFactory::RegisterWidgetConstructor(ITypedElementD
 	ActorCounter->Query = DataStorage.RegisterQuery(
 		Count().
 		Where().
-		All("/Script/MassActors.MassActorFragment"_Type).
+			All("/Script/MassActors.MassActorFragment"_Type).
 		Compile());
-	DataStorageUi.RegisterWidgetFactory(Purpose, MoveTemp(ActorCounter));
+	DataStorageUi.RegisterWidgetFactory(WigetPurpose, MoveTemp(ActorCounter));
 
 	TUniquePtr<FTypedElementCounterWidgetConstructor> WidgetCounter = MakeUnique<FTypedElementCounterWidgetConstructor>();
 	WidgetCounter->LabelText = LOCTEXT("WidgetCounterStatusBarLabel", "{0} {0}|plural(one=Widget, other=Widgets)");
@@ -84,21 +102,65 @@ void UTypedElementCounterWidgetFactory::RegisterWidgetConstructor(ITypedElementD
 	WidgetCounter->Query = DataStorage.RegisterQuery(
 		Count().
 		Where().
-		All<FTypedElementSlateWidgetReferenceColumn>().
+			All<FTypedElementSlateWidgetReferenceColumn>().
 		Compile());
-	DataStorageUi.RegisterWidgetFactory(Purpose, MoveTemp(WidgetCounter));
+	DataStorageUi.RegisterWidgetFactory(WigetPurpose, MoveTemp(WidgetCounter));
+}
+
+void UTypedElementCounterWidgetFactory::SetupMainWindowIntegrations(TSharedPtr<SWindow> ParentWindow, bool bIsRunningStartupDialog)
+{
+	UTypedElementRegistry* Registry = UTypedElementRegistry::GetInstance();
+	checkf(Registry, TEXT(
+		"FTypedElementsDataStorageUiModule didn't find the UTypedElementRegistry during main window integration when it should be available."));
+
+	ITypedElementDataStorageUiInterface* UiInterface = Registry->GetMutableDataStorageUi();
+	checkf(UiInterface, TEXT(
+		"FTypedElementsDataStorageUiModule tried to integrate with the main window before the "
+		"Typed Elements Data Storage UI interface is available."));
+
+	UToolMenu* Menu = UToolMenus::Get()->ExtendMenu(WigetPurpose);
+
+	TArray<TSharedRef<SWidget>> Widgets;
+	UiInterface->ConstructWidgets(WigetPurpose, {},
+		[&Widgets](const TSharedRef<SWidget>& NewWidget, TypedElementRowHandle Row)
+		{
+			Widgets.Add(NewWidget);
+		});
+
+	if (!Widgets.IsEmpty())
+	{
+		FToolMenuSection& Section = Menu->AddSection("DataStorageSection");
+		int32 WidgetCount = Widgets.Num();
+
+		Section.AddEntry(FToolMenuEntry::InitWidget("DataStorageStatusBarWidget_0", MoveTemp(Widgets[0]), FText::GetEmpty()));
+		for (int32 I = 1; I < WidgetCount; ++I)
+		{
+			Section.AddSeparator(FName(*FString::Format(TEXT("DataStorageStatusBarWidgetDivider_{0}"), { FString::FromInt(I) })));
+			Section.AddEntry(FToolMenuEntry::InitWidget(
+				FName(*FString::Format(TEXT("DataStorageStatusBarWidget_{0}"), { FString::FromInt(I) })),
+				MoveTemp(Widgets[I]), FText::GetEmpty()));
+		}
+	}
 }
 
 
+
+
 //
 // FTypedElementCounterWidgetConstructor
 //
 
+FTypedElementCounterWidgetConstructor::FTypedElementCounterWidgetConstructor()
+	: Super(FTypedElementCounterWidgetConstructor::StaticStruct())
+{
+}
+
 TSharedPtr<SWidget> FTypedElementCounterWidgetConstructor::CreateWidget()
 {
 	return SNew(STextBlock)
 		.Text(FText::Format(LabelText, 0))
 		.Margin(FMargin(4.0f, 0.0f))
+		.ToolTipText(ToolTipText)
 		.Justification(ETextJustify::Center);
 }
 
diff --git a/Engine/Plugins/Experimental/TypedElementsDataStorage/Source/TypedElementsDataStorageUI/Private/TypedElementsDataStorageUI.cpp b/Engine/Plugins/Experimental/TypedElementsDataStorage/Source/TypedElementsDataStorageUI/Private/TypedElementsDataStorageUI.cpp
--- a/Engine/Plugins/Experimental/TypedElementsDataStorage/Source/TypedElementsDataStorageUI/Private/TypedElementsDataStorageUI.cpp
+++ b/Engine/Plugins/Experimental/TypedElementsDataStorage/Source/TypedElementsDataStorageUI/Private/TypedElementsDataStorageUI.cpp
@@ -3,18 +3,33 @@
 #include "TypedElementsDataStorageUI.h"
 
 #include "Elements/Framework/TypedElementRegistry.h"
-#include "Interfaces/IMainFrameModule.h"
+#include "Elements/Interfaces/TypedElementDataStorageUiInterface.h"
+#include "HAL/IConsoleManager.h"
 #include "Misc/CoreDelegates.h"
-#include "ToolMenu.h"
-#include "ToolMenus.h"
 #include "UObject/UObjectGlobals.h"
-#include "Widgets/SWindow.h"
 
 #define LOCTEXT_NAMESPACE "FTypedElementsDataStorageModule"
 
+FAutoConsoleCommandWithOutputDevice PrintWidgetPurposesConsoleCommand(
+	TEXT("TEDS.UI.PrintWidgetPurposes"),
+	TEXT("Prints a list of all the known widget purposes."),
+	FConsoleCommandWithOutputDeviceDelegate::CreateLambda([](FOutputDevice& Output)
+		{
+			UTypedElementRegistry* Registry = UTypedElementRegistry::GetInstance();
+			if (ITypedElementDataStorageUiInterface* UiStorage = Registry->GetMutableDataStorageUi())
+			{
+				Output.Log(TEXT("The Typed Elements Data Storage has recorded the following widget purposes:"));
+				UiStorage->ListWidgetPurposes(
+					[&Output](FName Purpose, ITypedElementDataStorageUiInterface::EPurposeType, const FText& Description)
+					{
+						Output.Logf(TEXT("    %s - %s"), *Purpose.ToString(), *Description.ToString());
+					});
+				Output.Log(TEXT("End of Typed Elements Data Storage widget purpose list."));
+			}
+		}));
+
 void FTypedElementsDataStorageUiModule::StartupModule()
 {
-	IMainFrameModule::Get().OnMainFrameCreationFinished().AddStatic(&FTypedElementsDataStorageUiModule::SetupMainWindowIntegrations);
 }
 
 void FTypedElementsDataStorageUiModule::ShutdownModule()
@@ -30,42 +45,6 @@ FString FTypedElementsDataStorageUiModule::GetReferencerName() const
 	return TEXT("Typed Elements: Data Storage UI Module");
 }
 
-void FTypedElementsDataStorageUiModule::SetupMainWindowIntegrations(TSharedPtr<SWindow> ParentWindow, bool bIsRunningStartupDialog)
-{
-	UTypedElementRegistry* Registry = UTypedElementRegistry::GetInstance();
-	checkf(Registry, TEXT(
-		"FTypedElementsDataStorageUiModule didn't find the UTypedElementRegistry during main window integration when it should be available."));
-
-	ITypedElementDataStorageUiInterface* UiInterface = Registry->GetMutableDataStorageUi();
-	checkf(UiInterface, TEXT(
-		"FTypedElementsDataStorageUiModule tried to integrate with the main window before the "
-		"Typed Elements Data Storage UI interface is available."));
-	
-	UToolMenu* Menu = UToolMenus::Get()->ExtendMenu("LevelEditor.StatusBar.ToolBar");
-
-	TArray<TSharedRef<SWidget>> Widgets;
-	UiInterface->ConstructWidgets("LevelEditor.StatusBar.ToolBar", {},
-		[&Widgets](const TSharedRef<SWidget>& NewWidget, TypedElementRowHandle Row)
-		{
-			Widgets.Add(NewWidget);
-		});
-
-	if (!Widgets.IsEmpty())
-	{
-		FToolMenuSection& Section = Menu->AddSection("DataStorageSection");
-		int32 WidgetCount = Widgets.Num();
-
-		Section.AddEntry(FToolMenuEntry::InitWidget("DataStorageStatusBarWidget_0", MoveTemp(Widgets[0]), FText::GetEmpty()));
-		for (int32 I = 1; I < WidgetCount; ++I)
-		{
-			Section.AddSeparator(FName(*FString::Format(TEXT("DataStorageStatusBarWidgetDivider_{0}"), { FString::FromInt(I) })));
-			Section.AddEntry(FToolMenuEntry::InitWidget(
-				FName(*FString::Format(TEXT("DataStorageStatusBarWidget_{0}"), { FString::FromInt(I) })),
-				MoveTemp(Widgets[I]), FText::GetEmpty()));
-		}
-	}
-}
-
 #undef LOCTEXT_NAMESPACE
 	
 IMPLEMENT_MODULE(FTypedElementsDataStorageUiModule, TypedElementsDataStorageUI)
\ No newline at end of file
diff --git a/Engine/Plugins/Experimental/TypedElementsDataStorage/Source/TypedElementsDataStorageUI/Public/TypedElementCounterWidget.h b/Engine/Plugins/Experimental/TypedElementsDataStorage/Source/TypedElementsDataStorageUI/Public/TypedElementCounterWidget.h
--- a/Engine/Plugins/Experimental/TypedElementsDataStorage/Source/TypedElementsDataStorageUI/Public/TypedElementCounterWidget.h
+++ b/Engine/Plugins/Experimental/TypedElementsDataStorage/Source/TypedElementsDataStorageUI/Public/TypedElementCounterWidget.h
@@ -10,17 +10,26 @@
 
 #include "TypedElementCounterWidget.generated.h"
 
+class SWindow;
+
 UCLASS()
 class TYPEDELEMENTSDATASTORAGEUI_API UTypedElementCounterWidgetFactory : public UTypedElementDataStorageFactory
 {
 	GENERATED_BODY()
 
 public:
+	static FName WigetPurpose;
+
+	UTypedElementCounterWidgetFactory();
 	~UTypedElementCounterWidgetFactory() override = default;
 
 	void RegisterQueries(ITypedElementDataStorageInterface& DataStorage) const override;
-	void RegisterWidgetConstructor(ITypedElementDataStorageInterface& DataStorage,
+	void RegisterWidgetPurposes(ITypedElementDataStorageUiInterface& DataStorageUi) const override;
+	void RegisterWidgetConstructors(ITypedElementDataStorageInterface& DataStorage,
 		ITypedElementDataStorageUiInterface& DataStorageUi) const override;
+
+private:
+	static void SetupMainWindowIntegrations(TSharedPtr<SWindow> ParentWindow, bool bIsRunningStartupDialog);
 };
 
 /**
@@ -35,6 +44,7 @@ struct TYPEDELEMENTSDATASTORAGEUI_API FTypedElementCounterWidgetConstructor : pu
 	GENERATED_BODY()
 
 public:
+	FTypedElementCounterWidgetConstructor();
 	~FTypedElementCounterWidgetConstructor() override = default;
 
 	FText ToolTipText{ NSLOCTEXT("TypedElementUI_CounterWidget", "Tooltip", "Shows the total number found in the editor.") };
@@ -46,7 +56,7 @@ protected:
 	void AddColumns(ITypedElementDataStorageInterface* DataStorage, TypedElementRowHandle Row, const TSharedPtr<SWidget>& Widget) override;
 };
 
-USTRUCT()
+USTRUCT(meta = (DisplayName = "Counter widget"))
 struct TYPEDELEMENTSDATASTORAGEUI_API FTypedElementCounterWidgetColumn : public FTypedElementDataStorageColumn
 {
 	GENERATED_BODY()
diff --git a/Engine/Plugins/Experimental/TypedElementsDataStorage/Source/TypedElementsDataStorageUI/Public/TypedElementsDataStorageUI.h b/Engine/Plugins/Experimental/TypedElementsDataStorage/Source/TypedElementsDataStorageUI/Public/TypedElementsDataStorageUI.h
--- a/Engine/Plugins/Experimental/TypedElementsDataStorage/Source/TypedElementsDataStorageUI/Public/TypedElementsDataStorageUI.h
+++ b/Engine/Plugins/Experimental/TypedElementsDataStorage/Source/TypedElementsDataStorageUI/Public/TypedElementsDataStorageUI.h
@@ -29,7 +29,4 @@ public:
 	
 	void AddReferencedObjects(FReferenceCollector& Collector) override;
 	FString GetReferencerName() const override;
-
-private:
-	static void SetupMainWindowIntegrations(TSharedPtr<SWindow> ParentWindow, bool bIsRunningStartupDialog);
 };
diff --git a/Engine/Source/Runtime/TypedElementFramework/Private/Elements/Interfaces/TypedElementDataStorageUiInterface.cpp b/Engine/Source/Runtime/TypedElementFramework/Private/Elements/Interfaces/TypedElementDataStorageUiInterface.cpp
--- a/Engine/Source/Runtime/TypedElementFramework/Private/Elements/Interfaces/TypedElementDataStorageUiInterface.cpp
+++ b/Engine/Source/Runtime/TypedElementFramework/Private/Elements/Interfaces/TypedElementDataStorageUiInterface.cpp
@@ -4,6 +4,16 @@
 #include "Elements/Framework/TypedElementColumnUtils.h"
 #include "Elements/Columns/TypedElementSlateWidgetColumns.h"
 
+FTypedElementWidgetConstructor::FTypedElementWidgetConstructor(UScriptStruct* InTypeInfo)
+	: TypeInfo(InTypeInfo)
+{
+}
+
+const UScriptStruct* FTypedElementWidgetConstructor::GetTypeInfo() const
+{
+	return TypeInfo;
+}
+
 TSharedPtr<SWidget> FTypedElementWidgetConstructor::Construct(
 	TypedElementRowHandle Row,
 	ITypedElementDataStorageInterface* DataStorage,
diff --git a/Engine/Source/Runtime/TypedElementFramework/Public/Elements/Interfaces/TypedElementDataStorageFactory.h b/Engine/Source/Runtime/TypedElementFramework/Public/Elements/Interfaces/TypedElementDataStorageFactory.h
--- a/Engine/Source/Runtime/TypedElementFramework/Public/Elements/Interfaces/TypedElementDataStorageFactory.h
+++ b/Engine/Source/Runtime/TypedElementFramework/Public/Elements/Interfaces/TypedElementDataStorageFactory.h
@@ -19,9 +19,17 @@ class TYPEDELEMENTFRAMEWORK_API UTypedElementDataStorageFactory : public UObject
 public:
 	~UTypedElementDataStorageFactory() override = default;
 
+	/** 
+	 * Returns the order registration will be executed. Factories with a lower number will be executed
+	 * before factories with a higher number.
+	 */
+	virtual uint8 GetOrder() const { return 127; }
+
 	virtual void RegisterTables(ITypedElementDataStorageInterface& DataStorage) const {}
 	virtual void RegisterTickGroups(ITypedElementDataStorageInterface& DataStorage) const {}
 	virtual void RegisterQueries(ITypedElementDataStorageInterface& DataStorage) const {}
-	virtual void RegisterWidgetConstructor(ITypedElementDataStorageInterface& DataStorage,
+	
+	virtual void RegisterWidgetPurposes(ITypedElementDataStorageUiInterface& DataStorageUi) const {}
+	virtual void RegisterWidgetConstructors(ITypedElementDataStorageInterface& DataStorage,
 		ITypedElementDataStorageUiInterface& DataStorageUi) const {}
 };
\ No newline at end of file
diff --git a/Engine/Source/Runtime/TypedElementFramework/Public/Elements/Interfaces/TypedElementDataStorageUiInterface.h b/Engine/Source/Runtime/TypedElementFramework/Public/Elements/Interfaces/TypedElementDataStorageUiInterface.h
--- a/Engine/Source/Runtime/TypedElementFramework/Public/Elements/Interfaces/TypedElementDataStorageUiInterface.h
+++ b/Engine/Source/Runtime/TypedElementFramework/Public/Elements/Interfaces/TypedElementDataStorageUiInterface.h
@@ -29,8 +29,14 @@ struct TYPEDELEMENTFRAMEWORK_API FTypedElementWidgetConstructor
 	GENERATED_BODY()
 
 public:
+	explicit FTypedElementWidgetConstructor(UScriptStruct* InTypeInfo);
+	explicit FTypedElementWidgetConstructor(EForceInit) {} //< For compatibility and shouldn't be directly used.
+
 	virtual ~FTypedElementWidgetConstructor() = default;
 
+	/** Retrieves the type information for the constructor type. */
+	virtual const UScriptStruct* GetTypeInfo() const;
+
 	/**
 	 * Constructs the widget according to the provided information. Information is collected by calling
 	 * the below functions ApplyArguments, CreateWidget and AddColumns. It's recommeded to overload those
@@ -57,6 +63,8 @@ protected:
 	 * that keeps track of the lifecycle of the widget.
 	 */
 	virtual void AddColumns(ITypedElementDataStorageInterface* DataStorage, TypedElementRowHandle Row, const TSharedPtr<SWidget>& Widget);
+
+	UScriptStruct* TypeInfo{ nullptr };
 };
 
 template<>
@@ -64,6 +72,7 @@ struct TStructOpsTypeTraits<FTypedElementWidgetConstructor> : public TStructOpsT
 {
 	enum
 	{
+		WithNoInitConstructor = true,
 		WithPureVirtual = true,
 	};
 };
@@ -79,25 +88,78 @@ class TYPEDELEMENTFRAMEWORK_API ITypedElementDataStorageUiInterface
 	GENERATED_BODY()
 
 public:
+	enum class EPurposeType : uint8
+	{
+		/** General purpose name which allows multiple factory registrations. */
+		Generic,
+		/**
+		 * Only one factory can be registrated with this purpose. If multiple factories are registered only the last factory
+		 * will be stored.
+		 */
+		UniqueByName,
+		/**
+		 * Only one factory can be registrated with this purpose for a specific combination of columns. If multiple factories
+		 * are registered only the last factory will be stored.
+		 */
+		UniqueByNameAndColumn
+	};
+
 	using WidgetCreatedCallback = TFunctionRef<void(const TSharedRef<SWidget>& NewWidget, TypedElementRowHandle Row)>;
+	using WidgetConstructorCallback = TFunctionRef<bool(TUniquePtr<FTypedElementWidgetConstructor>, TConstArrayView<TWeakObjectPtr<const UScriptStruct>>)>;
+	using WidgetPurposeCallback = TFunctionRef<void(FName, EPurposeType, const FText&)>;
 
+	/**
+	 * Register a widget purpose. Widget purposes indicates how widgets can be used and categorizes/organizes the available 
+	 * widget factories. If the same purpose is registered multiple times, only the first will be recorded and later registrations
+	 * will be silently ignored.
+	 */
+	virtual void RegisterWidgetPurpose(FName Purpose, EPurposeType Type, FText Description) = 0;
 	/** 
 	 * Registers a widget factory that will be called when the purpose it's registered under is requested.
 	 * This version registers a generic type. Construction using these are typically cheaper as they can avoid
 	 * copying the Constructor and take up less memory. The downside is that they can't store additional configuration
-	 * options.
+	 * options. If the purpose has not been registered the factory will not be recorded and a warning will be printed.
+	 * If registration is successfull true will be returned otherwise false.
 	 */
-	virtual void RegisterWidgetFactory(FName Purpose, const UScriptStruct* Constructor) = 0;
-	
+	virtual bool RegisterWidgetFactory(FName Purpose, const UScriptStruct* Constructor) = 0;
+	/**
+	 * Registers a widget factory that will be called when the purpose it's registered under is requested.
+	 * This version registers a generic type. Construction using these are typically cheaper as they can avoid
+	 * copying the Constructor and take up less memory. The downside is that they can't store additional configuration
+	 * options. If the purpose has not been registered the factory will not be recorded and a warning will be printed.
+	 * The provided columns will be used when matching the factory during widget construction.
+	 * If registration is successfull true will be returned otherwise false.
+	 */
+	virtual bool RegisterWidgetFactory(FName Purpose, const UScriptStruct* Constructor,
+		TArray<TWeakObjectPtr<const UScriptStruct>> Columns) = 0;
 	/**
 	 * Registers a widget factory that will be called when the purpose it's registered under is requested.
 	 * This version uses a pre-created instance of the Constructor. The benefit of this is that it store
 	 * configuration options. The downside is that it takes up more memory and requires copying when it's
-	 * used.
+	 * used. If the purpose has not been registered the factory will not be recorded and a warning will be printed.
+	 * If registration is successfull true will be returned otherwise false.
 	 */
-	template<typename ConstructorType>
-	void RegisterWidgetFactory(FName Purpose, TUniquePtr<ConstructorType>&& Constructor);
+	virtual bool RegisterWidgetFactory(FName Purpose, TUniquePtr<FTypedElementWidgetConstructor>&& Constructor) = 0;
+	/**
+	 * Registers a widget factory that will be called when the purpose it's registered under is requested.
+	 * This version uses a pre-created instance of the Constructor. The benefit of this is that it store
+	 * configuration options. The downside is that it takes up more memory and requires copying when it's
+	 * used. If the purpose has not been registered the factory will not be recorded and a warning will be printed.
+	 * The provided columns will be used when matching the factory during widget construction.
+	 * If registration is successfull true will be returned otherwise false.
+	 */
+	virtual bool RegisterWidgetFactory(FName Purpose, TUniquePtr<FTypedElementWidgetConstructor>&& Constructor, TArray<TWeakObjectPtr<const UScriptStruct>> Columns) = 0;
 	
+	/** Creates widget constructors for the requested purpose. */
+	virtual void CreateWidgetConstructors(FName Purpose, 
+		TConstArrayView<TypedElement::ColumnUtils::Argument> Arguments, const WidgetConstructorCallback& Callback) = 0;
+	/** 
+	 * Finds matching widget constructors for provided columns, preferring longer matches over shorter matches.
+	 * The provided list of columns will be updated to contain all columns that couldn't be matched.
+	 */
+	virtual void CreateWidgetConstructors(FName Purpose, TArray<TWeakObjectPtr<const UScriptStruct>>& Columns,
+		TConstArrayView<TypedElement::ColumnUtils::Argument> Arguments, const WidgetConstructorCallback& Callback) = 0;
+
 	/**
 	 * Creates all the widgets registered under the provided name. This may be a large number of widgets for a wide name
 	 * or exactly one when the exact name of the widget is registered. Arguments can be provided, but widgets are free
@@ -110,23 +172,6 @@ public:
 	virtual TSharedPtr<SWidget> ConstructWidget(TypedElementRowHandle Row, FTypedElementWidgetConstructor& Constructor,
 		TConstArrayView<TypedElement::ColumnUtils::Argument> Arguments) = 0;
 
-protected:
-	/**
-	 * Registers a widget factory that will be called when the purpose it's registered under is requested.
-	 * This version uses a pre-created instance of the Constructor. The benefit of this is that it store
-	 * configuration options. The downside is that it takes up more memory and requires copying when it's
-	 * used.
-	 */
-	virtual void RegisterWidgetFactory(FName Purpose, const UScriptStruct* ConstructorType,
-		TUniquePtr<FTypedElementWidgetConstructor>&& Constructor) = 0;
+	/** Calls the provided callback for all known registered widget purposes. */
+	virtual void ListWidgetPurposes(const WidgetPurposeCallback& Callback) const = 0;
 };
-
-
-// Implementations
-template<typename ConstructorType>
-void ITypedElementDataStorageUiInterface::RegisterWidgetFactory(FName Purpose, TUniquePtr<ConstructorType>&& Constructor)
-{
-	static_assert(TIsDerivedFrom<ConstructorType, FTypedElementWidgetConstructor>::Value, 
-		"Provided Typed Element Widget Constructor doesn't derive from FTypedElementWidgetConstructor.");
-	RegisterWidgetFactory(Purpose, ConstructorType::StaticStruct(), MoveTemp(Constructor));
-}
\ No newline at end of file
