TEDS Actor->TypedElementDatabase Sync Optimization
* Track only actors which have been marked as changed by the World by watching PostEditChangeProperty
* This changes the performance scaling of TypedElementDatabase be based on number of actors that changed this frame instead of number of actors in the world
* Nominally, this will also make the inclusion of TypedElementDatabase nearly free for largely unchanging scenes at the expense of a modest cost to handle the PostEditChangeProperty
* Also added TypedElementDatabase specific profile macros which forward to trace and can be compiled out if required.  Conveniently add a searchable tag to the string literal passed as the event name for filtering in insights

diff --git a/Engine/Plugins/Experimental/TypedElementsDataStorage/Source/TypedElementsDataStorage/Private/Compatibility/TypedElementActorLabelProcessors.cpp b/Engine/Plugins/Experimental/TypedElementsDataStorage/Source/TypedElementsDataStorage/Private/Compatibility/TypedElementActorLabelProcessors.cpp
--- a/Engine/Plugins/Experimental/TypedElementsDataStorage/Source/TypedElementsDataStorage/Private/Compatibility/TypedElementActorLabelProcessors.cpp
+++ b/Engine/Plugins/Experimental/TypedElementsDataStorage/Source/TypedElementsDataStorage/Private/Compatibility/TypedElementActorLabelProcessors.cpp
@@ -99,6 +99,8 @@ void UTypedElementActorLabelFactory::RegisterActorLabelToColumnQuery(ITypedEleme
 				}
 			}
 		)
+		.Where()
+			.All<FTypedElementSyncFromWorldTag>()
 		.Compile()
 	);
 }
diff --git a/Engine/Plugins/Experimental/TypedElementsDataStorage/Source/TypedElementsDataStorage/Private/Compatibility/TypedElementActorTransformProcessors.cpp b/Engine/Plugins/Experimental/TypedElementsDataStorage/Source/TypedElementsDataStorage/Private/Compatibility/TypedElementActorTransformProcessors.cpp
--- a/Engine/Plugins/Experimental/TypedElementsDataStorage/Source/TypedElementsDataStorage/Private/Compatibility/TypedElementActorTransformProcessors.cpp
+++ b/Engine/Plugins/Experimental/TypedElementsDataStorage/Source/TypedElementsDataStorage/Private/Compatibility/TypedElementActorTransformProcessors.cpp
@@ -2,7 +2,6 @@
 
 #include "Compatibility/TypedElementActorTransformProcessors.h"
 
-
 #include "Elements/Columns/TypedElementMiscColumns.h"
 #include "Elements/Columns/TypedElementTransformColumns.h"
 #include "Elements/Framework/TypedElementQueryBuilder.h"
@@ -64,6 +63,8 @@ void UTypedElementActorTransformFactory::RegisterActorLocalTransformToColumn(ITy
 				}
 			}
 		)
+		.Where()
+			.All<FTypedElementSyncFromWorldTag>()
 		.Compile()
 	);
 }
diff --git a/Engine/Plugins/Experimental/TypedElementsDataStorage/Source/TypedElementsDataStorage/Private/Processors/TypedElementMiscProcessors.cpp b/Engine/Plugins/Experimental/TypedElementsDataStorage/Source/TypedElementsDataStorage/Private/Processors/TypedElementMiscProcessors.cpp
--- a/Engine/Plugins/Experimental/TypedElementsDataStorage/Source/TypedElementsDataStorage/Private/Processors/TypedElementMiscProcessors.cpp
+++ b/Engine/Plugins/Experimental/TypedElementsDataStorage/Source/TypedElementsDataStorage/Private/Processors/TypedElementMiscProcessors.cpp
@@ -23,4 +23,18 @@ void UTypedElementRemoveSyncToWorldTagFactory::RegisterQueries(ITypedElementData
 			.All<FTypedElementSyncBackToWorldTag>()
 		.Compile()
 	);
+
+	DataStorage.RegisterQuery(
+	Select(
+		TEXT("Remove 'sync from world' tag"),
+		FPhaseAmble(FPhaseAmble::ELocation::Postamble, DSI::EQueryTickPhase::FrameEnd),
+		[](DSI::IQueryContext& Context, const TypedElementRowHandle* Rows)
+		{
+			Context.RemoveColumns<FTypedElementSyncFromWorldTag>(TConstArrayView<TypedElementRowHandle>(Rows, Context.GetRowCount()));
+		}
+	)
+	.Where()
+		.All<FTypedElementSyncFromWorldTag>()
+	.Compile()
+);
 }
diff --git a/Engine/Plugins/Experimental/TypedElementsDataStorage/Source/TypedElementsDataStorage/Private/TypedElementDataStorageProfilingMacros.h b/Engine/Plugins/Experimental/TypedElementsDataStorage/Source/TypedElementsDataStorage/Private/TypedElementDataStorageProfilingMacros.h
new file mode 100644
--- /dev/null
+++ b/Engine/Plugins/Experimental/TypedElementsDataStorage/Source/TypedElementsDataStorage/Private/TypedElementDataStorageProfilingMacros.h
@@ -0,0 +1,21 @@
+﻿// Copyright Epic Games, Inc. All Rights Reserved.
+
+#pragma once
+
+#include "ProfilingDebugging/CpuProfilerTrace.h"
+
+#ifndef ENABLE_TEDS_PROFILING_MACROS
+#define ENABLE_TEDS_PROFILING_MACROS 1
+#endif
+
+#if ENABLE_TEDS_PROFILING_MACROS
+
+#define TEDS_EVENT_PREFIX TEXT("[TEDS] ")
+// Scoped event which prefixes the string literal with "[TEDS] " for consistent filtering
+#define TEDS_EVENT_SCOPE(StrLiteral) TRACE_CPUPROFILER_EVENT_SCOPE_STR(TEDS_EVENT_PREFIX StrLiteral)
+
+#else
+
+#define TEDS_EVENT_SCOPE(StrLiteral) do {} while(false)
+
+#endif
\ No newline at end of file
diff --git a/Engine/Plugins/Experimental/TypedElementsDataStorage/Source/TypedElementsDataStorage/Private/TypedElementDatabaseCompatibility.cpp b/Engine/Plugins/Experimental/TypedElementsDataStorage/Source/TypedElementsDataStorage/Private/TypedElementDatabaseCompatibility.cpp
--- a/Engine/Plugins/Experimental/TypedElementsDataStorage/Source/TypedElementsDataStorage/Private/TypedElementDatabaseCompatibility.cpp
+++ b/Engine/Plugins/Experimental/TypedElementsDataStorage/Source/TypedElementsDataStorage/Private/TypedElementDatabaseCompatibility.cpp
@@ -2,6 +2,7 @@
 
 #include "TypedElementDatabaseCompatibility.h"
 
+#include "Algo/Unique.h"
 #include "Editor.h"
 #include "Elements/Columns/TypedElementLabelColumns.h"
 #include "Elements/Columns/TypedElementMiscColumns.h"
@@ -9,19 +10,24 @@
 #include "Elements/Columns/TypedElementTransformColumns.h"
 #include "MassActorEditorSubsystem.h"
 #include "MassActorSubsystem.h"
+#include "TypedElementDataStorageProfilingMacros.h"
 
 void UTypedElementDatabaseCompatibility::Initialize(ITypedElementDataStorageInterface* StorageInterface)
 {
 	checkf(StorageInterface, TEXT("Typed Element's Database compatibility manager is being initialized with an invalid storage target."));
-
+	
 	Storage = StorageInterface;
 	Prepare();
 
-	StorageInterface->OnUpdate().AddUObject(this, &UTypedElementDatabaseCompatibility::AddPendingActors);
+	StorageInterface->OnUpdate().AddUObject(this, &UTypedElementDatabaseCompatibility::Tick);
+
+	PostEditChangePropertyDelegateHandle = FCoreUObjectDelegates::OnObjectPropertyChanged.AddUObject(this, &UTypedElementDatabaseCompatibility::OnPostEditChangeProperty);
 }
 
 void UTypedElementDatabaseCompatibility::Deinitialize()
 {
+	FCoreUObjectDelegates::OnObjectPropertyChanged.Remove(PostEditChangePropertyDelegateHandle);
+	
 	Reset();
 }
 
@@ -58,7 +64,7 @@ void UTypedElementDatabaseCompatibility::RemoveCompatibleObject(AActor* Actor)
 	}
 }
 
-TypedElementRowHandle UTypedElementDatabaseCompatibility::FindRowWithCompatibleObject(AActor* Actor) const
+TypedElementRowHandle UTypedElementDatabaseCompatibility::FindRowWithCompatibleObject(const TObjectKey<const AActor> Actor) const
 {
 	if (Storage && ActorSubsystem && Storage->IsAvailable())
 	{
@@ -93,7 +99,8 @@ void UTypedElementDatabaseCompatibility::CreateStandardArchetypes()
 			FTypedElementLabelColumn::StaticStruct(),
 			FTypedElementLabelHashColumn::StaticStruct(),
 			FTypedElementPackagePathColumn::StaticStruct(),
-			FTypedElementPackageLoadedPathColumn::StaticStruct()
+			FTypedElementPackageLoadedPathColumn::StaticStruct(),
+			FTypedElementSyncFromWorldTag::StaticStruct()
 		}), FName("Editor_StandardActorTable"));
 
 	StandardActorWithTransformTable = Storage->RegisterTable(StandardActorTable,
@@ -103,59 +110,129 @@ void UTypedElementDatabaseCompatibility::CreateStandardArchetypes()
 		}), FName("Editor_StandardActorWithTransformTable"));
 }
 
-void UTypedElementDatabaseCompatibility::AddPendingActors()
+void UTypedElementDatabaseCompatibility::Tick()
 {
+	TEDS_EVENT_SCOPE(TEXT("Compatibility Tick"))
 	UWorld* EditorWorld = GEditor->GetEditorWorldContext().World();
 
 	// Delay processing until the required systems are available by not clearing the pending actor list.
-	if (ActorsPendingRegistration.Num() > 0 && Storage && Storage->IsAvailable() && EditorWorld)
+	if (Storage && Storage->IsAvailable() && EditorWorld)
 	{
-		// Filter out the actors that are already registered or already destroyed. 
-		// The most common case for this is actors created from within MASS.
-		TWeakObjectPtr<AActor>* ActorBegin = ActorsPendingRegistration.GetData();
-		TWeakObjectPtr<AActor>* ActorIt = ActorsPendingRegistration.GetData();
-		TWeakObjectPtr<AActor>* ActorEnd = ActorBegin + ActorsPendingRegistration.Num();
-		while (ActorIt != ActorEnd)
+		if (ActorsPendingRegistration.Num() > 0)
 		{
-			if (
-				(*ActorIt).IsValid() && // If not valid or stale then the actor has already been deleted before it could be registered.
-				!ActorSubsystem->GetEntityHandleFromActor((*ActorIt).Get()).IsValid() && // If true, this actor is already registered with the Data Storage.
-				(*ActorIt)->GetWorld() == EditorWorld // Only record actors that are added to the editor to avoid taking PIE into account.
-				) 
+			// Filter out the actors that are already registered or already destroyed. 
+			// The most common case for this is actors created from within MASS.
+			TWeakObjectPtr<AActor>* ActorBegin = ActorsPendingRegistration.GetData();
+			TWeakObjectPtr<AActor>* ActorIt = ActorsPendingRegistration.GetData();
+			TWeakObjectPtr<AActor>* ActorEnd = ActorBegin + ActorsPendingRegistration.Num();
+			while (ActorIt != ActorEnd)
 			{
-				++ActorIt;
+				if (
+					(*ActorIt).IsValid() && // If not valid or stale then the actor has already been deleted before it could be registered.
+					!ActorSubsystem->GetEntityHandleFromActor((*ActorIt).Get()).IsValid() && // If true, this actor is already registered with the Data Storage.
+					(*ActorIt)->GetWorld() == EditorWorld // Only record actors that are added to the editor to avoid taking PIE into account.
+					) 
+				{
+					++ActorIt;
+				}
+				else
+				{
+					// Don't shrink the registration array as the array will be reused with a variety of different actor counts.
+					// If memory size becomes an issue it's better to resize the array once after this loop rather than within the
+					// loop to avoid many resizes happening.
+					constexpr bool bAllowToShrink = false;
+					ActorsPendingRegistration.RemoveAtSwap(ActorIt - ActorBegin, 1, bAllowToShrink);
+					--ActorEnd;
+				}
 			}
-			else
+
+			if (ActorBegin != ActorEnd) // Check if the number of actors to add isn't zero.
 			{
-				// Don't shrink the registration array as the array will be reused with a variety of different actor counts.
-				// If memory size becomes an issue it's better to resize the array once after this loop rather than within the
-				// loop to avoid many resizes happening.
-				constexpr bool bAllowToShrink = false;
-				ActorsPendingRegistration.RemoveAtSwap(ActorIt - ActorBegin, 1, bAllowToShrink);
-				--ActorEnd;
+				// Add the remaining actors to the data storage.
+				ActorIt = ActorsPendingRegistration.GetData();
+				Storage->BatchAddRow(StandardActorTable, ActorsPendingRegistration.Num(), [this, &ActorIt, ActorEnd](TypedElementRowHandle Row)
+					{
+						auto ActorStore = Storage->GetColumn<FMassActorFragment>(Row);
+						checkf(ActorStore, TEXT("Newly created row didn't contain the expected FMassActorFragment."));
+						
+						constexpr bool bIsOwnedByMass = false;
+					
+						AActor* Actor = ActorIt->Get();
+						check(Actor);
+						ActorStore->SetNoHandleMapUpdate(FMassEntityHandle::FromNumber(Row), Actor, bIsOwnedByMass);
+						ActorSubsystem->SetHandleForActor(ActorIt->Get(), FMassEntityHandle::FromNumber(Row));
+					
+						ActorsNeedingFullSync.Emplace(Actor);
+
+						checkf(ActorIt < ActorEnd,
+							TEXT("More (%i) entities were added than were requested (%i)."), ActorEnd - ActorIt, ActorsPendingRegistration.Num());
+						++ActorIt;
+					});
+
+				// Reset the container for next set of actors.
+				ActorsPendingRegistration.Reset();
 			}
 		}
 
-		if (ActorBegin != ActorEnd) // Check if the number of actors to add isn't zero.
+		if (!ActorsNeedingFullSync.IsEmpty())
 		{
-			// Add the remaining actors to the data storage.
-			ActorIt = ActorsPendingRegistration.GetData();
-			Storage->BatchAddRow(StandardActorTable, ActorsPendingRegistration.Num(), [this, &ActorIt, ActorEnd](TypedElementRowHandle Row)
-				{
-					auto ActorStore = Storage->GetColumn<FMassActorFragment>(Row);
-					checkf(ActorStore, TEXT("Newly created row didn't contain the expected FMassActorFragment."));
-					
-					constexpr bool bIsOwnedByMass = false;
-					ActorStore->SetNoHandleMapUpdate(FMassEntityHandle::FromNumber(Row), ActorIt->Get(), bIsOwnedByMass);
-					ActorSubsystem->SetHandleForActor(ActorIt->Get(), FMassEntityHandle::FromNumber(Row));
+			TEDS_EVENT_SCOPE(TEXT("Process ActorsNeedingFullSync"));
+			// Deduplicate to avoid duplicate reverse lookups or adding tags more than once
+			{
+				TEDS_EVENT_SCOPE(TEXT("Deduplicate ActorsNeedingFullSync"));
+				ActorsNeedingFullSync.Sort();
+				Algo::Unique(ActorsNeedingFullSync);
+			}
 
-					checkf(ActorIt < ActorEnd,
-						TEXT("More (%i) entities were added than were requested (%i)."), ActorEnd - ActorIt, ActorsPendingRegistration.Num());
-					++ActorIt;
-				});
+			TArray<TypedElementRowHandle> RowHandles;
+			{
+				TEDS_EVENT_SCOPE(TEXT("Reverse lookup Rows from Actors"));
+				
+				RowHandles.SetNumUninitialized(ActorsNeedingFullSync.Num());
+				{
+					int32 RowHandleIndex = 0;
+					for (int32 ActorIndex = 0, End = ActorsNeedingFullSync.Num(); ActorIndex < End; ++ActorIndex)
+					{
+						const TObjectKey<const AActor> ActorKey = ActorsNeedingFullSync[ActorIndex];
+						const TypedElementRowHandle Row = FindRowWithCompatibleObject(ActorKey);
+						if (Row != TypedElementInvalidRowHandle)
+						{
+							RowHandles[RowHandleIndex++] = Row;
+						}
+					}
+					const int32 RowHandleCount = RowHandleIndex;
+					const bool bAllowShrinking = false;
+					RowHandles.SetNum(RowHandleCount, bAllowShrinking);
+				}
+
+				ActorsNeedingFullSync.Reset();
+			}
 
-			// Reset the container for next set of actors.
-			ActorsPendingRegistration.Reset();
+			{
+				TEDS_EVENT_SCOPE(TEXT("Add SyncToWorld Tag"));
+				// Tag the rows containing actor data that they need to be synced
+				// Note: Watch out for the performance of this, may end up doing a lot of row moves
+				for (TypedElementRowHandle Row : RowHandles)
+				{
+					Storage->AddTag(Row, FTypedElementSyncFromWorldTag::StaticStruct());
+				}
+			}
 		}
 	}
-}
\ No newline at end of file
+}
+
+void UTypedElementDatabaseCompatibility::OnPostEditChangeProperty(
+	UObject* Object,
+	FPropertyChangedEvent& /*PropertyChangedEvent*/)
+{
+	// We aren't sure if this Actor is tracked by the database
+	// Will resolve that during the tick step
+	// The aim is to keep this delegate handler as simple as possible to avoid performance side-effects when other code
+	// invokes it.
+	AActor* Actor = Cast<AActor>(Object);
+	if (Actor)
+	{
+		// Note: This array may end up with duplicates
+		ActorsNeedingFullSync.Add(Actor);
+	}
+}
diff --git a/Engine/Plugins/Experimental/TypedElementsDataStorage/Source/TypedElementsDataStorage/Private/TypedElementDatabaseCompatibility.h b/Engine/Plugins/Experimental/TypedElementsDataStorage/Source/TypedElementsDataStorage/Private/TypedElementDatabaseCompatibility.h
--- a/Engine/Plugins/Experimental/TypedElementsDataStorage/Source/TypedElementsDataStorage/Private/TypedElementDatabaseCompatibility.h
+++ b/Engine/Plugins/Experimental/TypedElementsDataStorage/Source/TypedElementsDataStorage/Private/TypedElementDatabaseCompatibility.h
@@ -28,14 +28,16 @@ public:
 
 	void AddCompatibleObject(AActor* Actor) override;
 	void RemoveCompatibleObject(AActor* Actor) override;
-	TypedElementRowHandle FindRowWithCompatibleObject(AActor* Actor) const;
+	TypedElementRowHandle FindRowWithCompatibleObject(const TObjectKey<const AActor> Actor) const override;
 
 private:
 	void Prepare();
 	void Reset();
 	void CreateStandardArchetypes();
 	
-	void AddPendingActors();
+	void Tick();
+
+	void OnPostEditChangeProperty(UObject* Object, FPropertyChangedEvent& PropertyChangedEvent);
 	
 	TArray<TWeakObjectPtr<AActor>> ActorsPendingRegistration;
 	
@@ -43,4 +45,13 @@ private:
 	TypedElementTableHandle StandardActorWithTransformTable{ TypedElementInvalidTableHandle };
 	ITypedElementDataStorageInterface* Storage{ nullptr };
 	TSharedPtr<FMassActorManager> ActorSubsystem;
+
+	/**~
+	 * Reference of actors that need to be fully synced from the world to the database
+	 * May have duplicates
+	 * Caution: Could point to actors that have been GCd
+	 */
+	TArray<TObjectKey<const AActor>> ActorsNeedingFullSync;
+
+	FDelegateHandle PostEditChangePropertyDelegateHandle;
 };
\ No newline at end of file
diff --git a/Engine/Source/Runtime/TypedElementFramework/Public/Elements/Columns/TypedElementMiscColumns.h b/Engine/Source/Runtime/TypedElementFramework/Public/Elements/Columns/TypedElementMiscColumns.h
--- a/Engine/Source/Runtime/TypedElementFramework/Public/Elements/Columns/TypedElementMiscColumns.h
+++ b/Engine/Source/Runtime/TypedElementFramework/Public/Elements/Columns/TypedElementMiscColumns.h
@@ -19,6 +19,18 @@ struct FTypedElementSyncBackToWorldTag final : public FTypedElementDataStorageTa
 	GENERATED_BODY()
 };
 
+/**
+ * Tag to signal that data a processor copies out of the world must be synced to the datastorage.
+ * Useful for when an Actor was recently spawned or reloaded in the world.
+ * Currently used if any property changes since there is no mechanism to selectively run
+ * queries for specific changed properties.
+ */
+USTRUCT(meta = (DisplayName = "Sync from world"))
+struct FTypedElementSyncFromWorldTag final : public FTypedElementDataStorageTag
+{
+	GENERATED_BODY()
+};
+
 USTRUCT(meta = (DisplayName = "Row reference"))
 struct FTypedElementRowReferenceColumn final : public FTypedElementDataStorageColumn
 {
diff --git a/Engine/Source/Runtime/TypedElementFramework/Public/Elements/Interfaces/TypedElementDataStorageCompatibilityInterface.h b/Engine/Source/Runtime/TypedElementFramework/Public/Elements/Interfaces/TypedElementDataStorageCompatibilityInterface.h
--- a/Engine/Source/Runtime/TypedElementFramework/Public/Elements/Interfaces/TypedElementDataStorageCompatibilityInterface.h
+++ b/Engine/Source/Runtime/TypedElementFramework/Public/Elements/Interfaces/TypedElementDataStorageCompatibilityInterface.h
@@ -4,6 +4,7 @@
 
 #include "Elements/Interfaces/TypedElementDataStorageInterface.h"
 #include "UObject/Interface.h"
+#include "UObject/ObjectKey.h"
 #include "UObject/ObjectMacros.h"
 
 #include "TypedElementDataStorageCompatibilityInterface.generated.h"
@@ -27,5 +28,5 @@ class TYPEDELEMENTFRAMEWORK_API ITypedElementDataStorageCompatibilityInterface
 public:
 	virtual void AddCompatibleObject(AActor* Actor) = 0;
 	virtual void RemoveCompatibleObject(AActor* Actor) = 0;
-	virtual TypedElementRowHandle FindRowWithCompatibleObject(AActor* Actor) const = 0;
+	virtual TypedElementRowHandle FindRowWithCompatibleObject(const TObjectKey<const AActor> Actor) const = 0;
 };
\ No newline at end of file
