Set size of uniqued array in TypedElementDatabaseCompatibility
* Algo::Unique will shift to-be-deleted elements to the end of the array but not change the size of the array.
* The existing code results in an unsorted list that the subsequent code still loops through.
* Fix is to truncate the list at the point where duplicates were swapped to.
* Fixed typo in event name

diff --git a/Engine/Plugins/Experimental/TypedElementsDataStorage/Source/TypedElementsDataStorage/Private/TypedElementDatabaseCompatibility.cpp b/Engine/Plugins/Experimental/TypedElementsDataStorage/Source/TypedElementsDataStorage/Private/TypedElementDatabaseCompatibility.cpp
--- a/Engine/Plugins/Experimental/TypedElementsDataStorage/Source/TypedElementsDataStorage/Private/TypedElementDatabaseCompatibility.cpp
+++ b/Engine/Plugins/Experimental/TypedElementsDataStorage/Source/TypedElementsDataStorage/Private/TypedElementDatabaseCompatibility.cpp
@@ -181,7 +181,7 @@ void UTypedElementDatabaseCompatibility::Tick()
 			{
 				TEDS_EVENT_SCOPE(TEXT("Deduplicate ActorsNeedingFullSync"));
 				ActorsNeedingFullSync.Sort();
-				Algo::Unique(ActorsNeedingFullSync);
+				ActorsNeedingFullSync.SetNum(Algo::Unique(ActorsNeedingFullSync));
 			}
 
 			TArray<TypedElementRowHandle> RowHandles;
@@ -209,7 +209,7 @@ void UTypedElementDatabaseCompatibility::Tick()
 			}
 
 			{
-				TEDS_EVENT_SCOPE(TEXT("Add SyncToWorld Tag"));
+				TEDS_EVENT_SCOPE(TEXT("Add SyncFromWorld Tag"));
 				// Tag the rows containing actor data that they need to be synced
 				// Note: Watch out for the performance of this, may end up doing a lot of row moves
 				for (TypedElementRowHandle Row : RowHandles)
