Did a pass over the columns in TEDS.

The most notable changes are:
- AddTag has been replaced with the more generic AddColumn to be in line with the overal philosophy of TEDS.
- It's now possible to add/remove columns in a batch for a row or a batch of rows. This helps improve performance as it reduces the amount of moving TEDS has to do between tables. As a bonus it also helps reduce the number of "temp" tables in the Mass Debugger that clutter the view.
- It's now possible to check if a row has specific columns set on it.

#nojira

diff --git a/Engine/Plugins/Experimental/TypedElementsDataStorage/Source/TypedElementsDataStorage/Private/TypedElementDatabase.cpp b/Engine/Plugins/Experimental/TypedElementsDataStorage/Source/TypedElementsDataStorage/Private/TypedElementDatabase.cpp
--- a/Engine/Plugins/Experimental/TypedElementsDataStorage/Source/TypedElementsDataStorage/Private/TypedElementDatabase.cpp
+++ b/Engine/Plugins/Experimental/TypedElementsDataStorage/Source/TypedElementsDataStorage/Private/TypedElementDatabase.cpp
@@ -206,42 +206,73 @@ void UTypedElementDatabase::RemoveRow(TypedElementRowHandle Row)
 	}
 }
 
-void UTypedElementDatabase::AddTag(TypedElementRowHandle Row, const UScriptStruct* TagType)
+bool UTypedElementDatabase::AddColumn(TypedElementRowHandle Row, const UScriptStruct* ColumnType)
 {
-	checkf(TagType && TagType->IsChildOf(FMassTag::StaticStruct()),
-		TEXT("Tag type '%s' is invalid as it needs to be set or derived from FMassTag."), *GetPathNameSafe(TagType));
+	FMassEntityHandle Entity = FMassEntityHandle::FromNumber(Row);
+	if (ActiveEditorEntityManager && ActiveEditorEntityManager->IsEntityValid(Entity))
+	{
+		if (ColumnType->IsChildOf(FMassTag::StaticStruct()))
+		{
+			ActiveEditorEntityManager->AddTagToEntity(Entity, ColumnType);
+			return true;
+		}
+		else if (ColumnType->IsChildOf(FMassFragment::StaticStruct()))
+		{
+			FStructView Column = ActiveEditorEntityManager->GetFragmentDataStruct(Entity, ColumnType);
+			if (!Column.IsValid())
+			{
+				ActiveEditorEntityManager->AddFragmentToEntity(Entity, ColumnType);
+				return true;
+			}
+		}
+	}
+	return false;
+}
+
+bool UTypedElementDatabase::AddColumn(TypedElementRowHandle Row, FTopLevelAssetPath ColumnName)
+{
+	bool bExactMatch = true;
+	UScriptStruct* ColumnStructInfo = Cast<UScriptStruct>(StaticFindObject(UScriptStruct::StaticClass(), ColumnName, bExactMatch));
+	return ColumnStructInfo ? AddColumn(Row, ColumnStructInfo) : false;
+}
 
+void UTypedElementDatabase::RemoveColumn(TypedElementRowHandle Row, const UScriptStruct* ColumnType)
+{
 	FMassEntityHandle Entity = FMassEntityHandle::FromNumber(Row);
 	if (ActiveEditorEntityManager && ActiveEditorEntityManager->IsEntityValid(Entity))
 	{
-		ActiveEditorEntityManager->AddTagToEntity(Entity, TagType);
+		if (ColumnType->IsChildOf(FMassTag::StaticStruct()))
+		{
+			ActiveEditorEntityManager->RemoveTagFromEntity(Entity, ColumnType);
+		}
+		else if (ColumnType->IsChildOf(FMassFragment::StaticStruct()))
+		{
+			ActiveEditorEntityManager->RemoveFragmentFromEntity(Entity, ColumnType);
+		}
 	}
 }
 
-void UTypedElementDatabase::AddTag(TypedElementRowHandle Row, FTopLevelAssetPath TagName)
+void UTypedElementDatabase::RemoveColumn(TypedElementRowHandle Row, FTopLevelAssetPath ColumnName)
 {
 	bool bExactMatch = true;
-	UScriptStruct* TagStructInfo = Cast<UScriptStruct>(StaticFindObject(UScriptStruct::StaticClass(), TagName, bExactMatch));
-	if (TagStructInfo)
+	if (UScriptStruct* ColumnStructInfo = Cast<UScriptStruct>(StaticFindObject(UScriptStruct::StaticClass(), ColumnName, bExactMatch)))
 	{
-		AddTag(Row, TagStructInfo);
+		RemoveColumn(Row, ColumnStructInfo);
 	}
 }
 
 void* UTypedElementDatabase::AddOrGetColumnData(TypedElementRowHandle Row, const UScriptStruct* ColumnType)
 {
-	checkf(ColumnType && ColumnType->IsChildOf(FMassFragment::StaticStruct()), 
-		TEXT("Colum type '%s' is invalid as it needs to be set or derived from FMassFragment."), *GetPathNameSafe(ColumnType));
-
 	FMassEntityHandle Entity = FMassEntityHandle::FromNumber(Row);
-	if (ActiveEditorEntityManager && ActiveEditorEntityManager->IsEntityValid(Entity))
+	if (ActiveEditorEntityManager && ActiveEditorEntityManager->IsEntityValid(Entity) &&
+		ColumnType && ColumnType->IsChildOf(FMassFragment::StaticStruct()))
 	{
 		FStructView Column = ActiveEditorEntityManager->GetFragmentDataStruct(Entity, ColumnType);
 		if (!Column.IsValid())
 		{
 			ActiveEditorEntityManager->AddFragmentToEntity(Entity, ColumnType);
 			Column = ActiveEditorEntityManager->GetFragmentDataStruct(Entity, ColumnType);
-			checkf(Column.IsValid(), TEXT("Added a new column to the Typed Element's data storae, but it couldn't be retrieved."));
+			checkf(Column.IsValid(), TEXT("Added a new column to the Typed Element's data store, but it couldn't be retrieved."));
 
 		}
 		return Column.GetMemory();
@@ -275,11 +306,9 @@ ColumnDataResult UTypedElementDatabase::AddOrGetColumnData(TypedElementRowHandle
 
 void* UTypedElementDatabase::GetColumnData(TypedElementRowHandle Row, const UScriptStruct* ColumnType)
 {
-	checkf(ColumnType && ColumnType->IsChildOf(FMassFragment::StaticStruct()),
-		TEXT("Colum type '%s' is invalid as it needs to be set or derived from FMassFragment."), *GetPathNameSafe(ColumnType));
-
 	FMassEntityHandle Entity = FMassEntityHandle::FromNumber(Row);
-	if (ActiveEditorEntityManager && ActiveEditorEntityManager->IsEntityValid(Entity))
+	if (ActiveEditorEntityManager && ActiveEditorEntityManager->IsEntityValid(Entity) &&
+		ColumnType && ColumnType->IsChildOf(FMassFragment::StaticStruct()))
 	{
 		FStructView Column = ActiveEditorEntityManager->GetFragmentDataStruct(Entity, ColumnType);
 		if (Column.IsValid())
@@ -306,7 +335,7 @@ ColumnDataResult UTypedElementDatabase::GetColumnData(TypedElementRowHandle Row,
 				}
 			});
 
-		if (FragmentType != nullptr)
+		if (FragmentType && FragmentType->IsChildOf(FMassFragment::StaticStruct()))
 		{
 			FStructView Column = ActiveEditorEntityManager->GetFragmentDataStruct(Entity, FragmentType);
 			if (Column.IsValid())
@@ -318,6 +347,196 @@ ColumnDataResult UTypedElementDatabase::GetColumnData(TypedElementRowHandle Row,
 	return ColumnDataResult{ nullptr, nullptr };
 }
 
+bool UTypedElementDatabase::AddColumns(TypedElementRowHandle Row, TConstArrayView<const UScriptStruct*> Columns)
+{
+	FMassEntityHandle Entity = FMassEntityHandle::FromNumber(Row);
+	if (ActiveEditorEntityManager && ActiveEditorEntityManager->IsEntityValid(Entity))
+	{
+		FMassArchetypeHandle Archetype = ActiveEditorEntityManager->GetArchetypeForEntity(Entity);
+
+		FMassFragmentBitSet FragmentsToAdd;
+		FMassTagBitSet TagsToAdd;
+		if (ColumnsToBitSets(Columns, FragmentsToAdd, TagsToAdd))
+		{
+			FMassArchetypeCompositionDescriptor AddComposition(
+				MoveTemp(FragmentsToAdd), MoveTemp(TagsToAdd), FMassChunkFragmentBitSet(), FMassSharedFragmentBitSet());
+			ActiveEditorEntityManager->AddCompositionToEntity_GetDelta(Entity, AddComposition);
+			return true;
+		}
+	}
+	return false;
+}
+
+void UTypedElementDatabase::RemoveColumns(TypedElementRowHandle Row, TConstArrayView<const UScriptStruct*> Columns)
+{
+	FMassEntityHandle Entity = FMassEntityHandle::FromNumber(Row);
+	if (ActiveEditorEntityManager && ActiveEditorEntityManager->IsEntityValid(Entity))
+	{
+		FMassArchetypeHandle Archetype = ActiveEditorEntityManager->GetArchetypeForEntity(Entity);
+
+		FMassFragmentBitSet FragmentsToRemove;
+		FMassTagBitSet TagsToRemove;
+		if (ColumnsToBitSets(Columns, FragmentsToRemove, TagsToRemove))
+		{
+			FMassArchetypeCompositionDescriptor RemoveComposition(
+				MoveTemp(FragmentsToRemove), MoveTemp(TagsToRemove), FMassChunkFragmentBitSet(), FMassSharedFragmentBitSet());
+			ActiveEditorEntityManager->RemoveCompositionFromEntity(Entity, RemoveComposition);
+		}
+	}
+}
+
+bool UTypedElementDatabase::AddRemoveColumns(TypedElementRowHandle Row,
+	TConstArrayView<const UScriptStruct*> ColumnsToAdd, TConstArrayView<const UScriptStruct*> ColumnsToRemove)
+{
+	bool bResult = false;
+	FMassEntityHandle Entity = FMassEntityHandle::FromNumber(Row);
+	if (ActiveEditorEntityManager && ActiveEditorEntityManager->IsEntityValid(Entity))
+	{
+		FMassArchetypeHandle Archetype = ActiveEditorEntityManager->GetArchetypeForEntity(Entity);
+
+		FMassFragmentBitSet FragmentsToAdd;
+		FMassTagBitSet TagsToAdd;
+		if (ColumnsToBitSets(ColumnsToAdd, FragmentsToAdd, TagsToAdd))
+		{
+			FMassArchetypeCompositionDescriptor AddComposition(
+				MoveTemp(FragmentsToAdd), MoveTemp(TagsToAdd), FMassChunkFragmentBitSet(), FMassSharedFragmentBitSet());
+			ActiveEditorEntityManager->AddCompositionToEntity_GetDelta(Entity, AddComposition);
+			bResult = true;
+		}
+
+		FMassTagBitSet TagsToRemove;
+		FMassFragmentBitSet FragmentsToRemove;
+		if (ColumnsToBitSets(ColumnsToRemove, FragmentsToRemove, TagsToRemove))
+		{
+			FMassArchetypeCompositionDescriptor RemoveComposition(
+				MoveTemp(FragmentsToRemove), MoveTemp(TagsToRemove), FMassChunkFragmentBitSet(), FMassSharedFragmentBitSet());
+			ActiveEditorEntityManager->RemoveCompositionFromEntity(Entity, RemoveComposition);
+			bResult = true;
+		}
+	}
+	return bResult;
+}
+
+bool UTypedElementDatabase::BatchAddRemoveColumns(TConstArrayView<TypedElementRowHandle> Rows, 
+	TConstArrayView<const UScriptStruct*> ColumnsToAdd, TConstArrayView<const UScriptStruct*> ColumnsToRemove)
+{
+	if (ActiveEditorEntityManager)
+	{
+		FMassFragmentBitSet FragmentsToAdd;
+		FMassFragmentBitSet FragmentsToRemove;
+
+		FMassTagBitSet TagsToAdd;
+		FMassTagBitSet TagsToRemove;
+
+		bool bMustUpdateFragments = ColumnsToBitSets(ColumnsToAdd, FragmentsToAdd, TagsToAdd);
+		bool bMustUpdateTags = ColumnsToBitSets(ColumnsToRemove, FragmentsToRemove, TagsToRemove);
+		
+		if (bMustUpdateFragments || bMustUpdateTags)
+		{
+			using EntityHandleArray = TArray<FMassEntityHandle, TInlineAllocator<32>>;
+			using EntityArchetypeLookup = TMap<FMassArchetypeHandle, EntityHandleArray, TInlineSetAllocator<32>>;
+			using ArchetypeEntityArray = TArray<FMassArchetypeEntityCollection, TInlineAllocator<32>>;
+
+			// Sort rows (entities) into to matching table (archetype) bucket.
+			EntityArchetypeLookup LookupTable;
+			for (TypedElementRowHandle EntityId : Rows)
+			{
+				FMassEntityHandle Entity = FMassEntityHandle::FromNumber(EntityId);
+				if (ActiveEditorEntityManager->IsEntityValid(Entity))
+				{
+					FMassArchetypeHandle Archetype = ActiveEditorEntityManager->GetArchetypeForEntity(Entity);
+					EntityHandleArray& EntityCollection = LookupTable.FindOrAdd(Archetype);
+					EntityCollection.Add(Entity);
+				}
+			}
+			
+			// Construct table (archetype) specific row (entity) collections.
+			ArchetypeEntityArray EntityCollections;
+			EntityCollections.Reserve(LookupTable.Num());
+			for (auto It = LookupTable.CreateConstIterator(); It; ++It)
+			{
+				EntityCollections.Emplace(It.Key(), It.Value(), FMassArchetypeEntityCollection::EDuplicatesHandling::FoldDuplicates);
+			}
+
+			// Batch update usint the appropriate fragment/bit sets.
+			if (bMustUpdateFragments)
+			{
+				ActiveEditorEntityManager->BatchChangeFragmentCompositionForEntities(EntityCollections, FragmentsToAdd, FragmentsToRemove);
+			}
+			if (bMustUpdateTags)
+			{
+				ActiveEditorEntityManager->BatchChangeTagsForEntities(EntityCollections, TagsToAdd, TagsToRemove);
+			}
+			return true;
+		}
+	}
+	return false;
+}
+
+bool UTypedElementDatabase::HasColumns(TypedElementRowHandle Row, TConstArrayView<const UScriptStruct*> ColumnTypes) const
+{
+	FMassEntityHandle Entity = FMassEntityHandle::FromNumber(Row);
+	if (ActiveEditorEntityManager && ActiveEditorEntityManager->IsEntityValid(Entity))
+	{
+		FMassArchetypeHandle Archetype = ActiveEditorEntityManager->GetArchetypeForEntity(Entity);
+		const FMassArchetypeCompositionDescriptor& Composition = ActiveEditorEntityManager->GetArchetypeComposition(Archetype);
+
+		bool bHasAllColumns = true;
+		const UScriptStruct* const* ColumnTypesEnd = ColumnTypes.end();
+		for (const UScriptStruct* const* ColumnType = ColumnTypes.begin(); ColumnType != ColumnTypesEnd && bHasAllColumns; ++ColumnType)
+		{
+			if ((*ColumnType)->IsChildOf(FMassFragment::StaticStruct()))
+			{
+				bHasAllColumns = Composition.Fragments.Contains(**ColumnType);
+			}
+			else if ((*ColumnType)->IsChildOf(FMassTag::StaticStruct()))
+			{
+				bHasAllColumns = Composition.Tags.Contains(**ColumnType);
+			}
+			else
+			{
+				return false;
+			}
+		}
+
+		return bHasAllColumns;
+	}
+	return false;
+}
+
+bool UTypedElementDatabase::HasColumns(TypedElementRowHandle Row, TConstArrayView<TWeakObjectPtr<const UScriptStruct>> ColumnTypes) const
+{
+	FMassEntityHandle Entity = FMassEntityHandle::FromNumber(Row);
+	if (ActiveEditorEntityManager && ActiveEditorEntityManager->IsEntityValid(Entity))
+	{
+		FMassArchetypeHandle Archetype = ActiveEditorEntityManager->GetArchetypeForEntity(Entity);
+		const FMassArchetypeCompositionDescriptor& Composition = ActiveEditorEntityManager->GetArchetypeComposition(Archetype);
+
+		bool bHasAllColumns = true;
+		const TWeakObjectPtr<const UScriptStruct>* ColumnTypesEnd = ColumnTypes.end();
+		for (const TWeakObjectPtr<const UScriptStruct>* ColumnType = ColumnTypes.begin(); ColumnType != ColumnTypesEnd && bHasAllColumns; ++ColumnType)
+		{
+			if (ColumnType->IsValid())
+			{
+				if ((*ColumnType)->IsChildOf(FMassFragment::StaticStruct()))
+				{
+					bHasAllColumns = Composition.Fragments.Contains(**ColumnType);
+					continue;
+				}
+				else if ((*ColumnType)->IsChildOf(FMassTag::StaticStruct()))
+				{
+					bHasAllColumns = Composition.Tags.Contains(**ColumnType);
+					continue;
+				}
+			}
+			return false;
+		}
+
+		return bHasAllColumns;
+	}
+	return false;
+}
+
 void UTypedElementDatabase::RegisterTickGroup(
 	FName GroupName, EQueryTickPhase Phase, FName BeforeGroup, FName AfterGroup, bool bRequiresMainThread)
 {
@@ -783,6 +1002,25 @@ void* UTypedElementDatabase::GetExternalSystemAddress(UClass* Target)
 	return nullptr;
 }
 
+bool UTypedElementDatabase::ColumnsToBitSets(TConstArrayView<const UScriptStruct*> Columns, FMassFragmentBitSet& Fragments, FMassTagBitSet& Tags)
+{
+	bool bResult = false;
+	for (const UScriptStruct* ColumnType : Columns)
+	{
+		if (ColumnType->IsChildOf(FMassFragment::StaticStruct()))
+		{
+			Fragments.Add(*ColumnType);
+			bResult = true;
+		}
+		else if (ColumnType->IsChildOf(FMassTag::StaticStruct()))
+		{
+			Tags.Add(*ColumnType);
+			bResult = true;
+		}
+	}
+	return bResult;
+}
+
 void UTypedElementDatabase::PreparePhase(EQueryTickPhase Phase, float DeltaTime)
 {
 	PhasePreOrPostAmble(Phase, DeltaTime, PhasePreparationQueries[static_cast<std::underlying_type_t<EQueryTickPhase>>(Phase)]);
diff --git a/Engine/Plugins/Experimental/TypedElementsDataStorage/Source/TypedElementsDataStorage/Private/TypedElementDatabase.h b/Engine/Plugins/Experimental/TypedElementsDataStorage/Source/TypedElementsDataStorage/Private/TypedElementDatabase.h
--- a/Engine/Plugins/Experimental/TypedElementsDataStorage/Source/TypedElementsDataStorage/Private/TypedElementDatabase.h
+++ b/Engine/Plugins/Experimental/TypedElementsDataStorage/Source/TypedElementsDataStorage/Private/TypedElementDatabase.h
@@ -57,14 +57,24 @@ public:
 	bool BatchAddRow(FName TableName, int32 Count, TypedElementDataStorageCreationCallbackRef OnCreated) override;
 	void RemoveRow(TypedElementRowHandle Row) override;
 
-	void AddTag(TypedElementRowHandle Row, const UScriptStruct* TagType) override;
-	void AddTag(TypedElementRowHandle Row, FTopLevelAssetPath TagName) override;
+	bool AddColumn(TypedElementRowHandle Row, const UScriptStruct* ColumnType) override;
+	bool AddColumn(TypedElementRowHandle Row, FTopLevelAssetPath ColumnName) override;
+	void RemoveColumn(TypedElementRowHandle Row, const UScriptStruct* ColumnType) override;
+	void RemoveColumn(TypedElementRowHandle Row, FTopLevelAssetPath ColumnName) override;
 	void* AddOrGetColumnData(TypedElementRowHandle Row, const UScriptStruct* ColumnType) override;
 	ColumnDataResult AddOrGetColumnData(TypedElementRowHandle Row, FTopLevelAssetPath ColumnName) override;
 	void* GetColumnData(TypedElementRowHandle Row, const UScriptStruct* ColumnType) override;
 	ColumnDataResult AddOrGetColumnData(TypedElementRowHandle Row, FTopLevelAssetPath ColumnName,
 		TConstArrayView<TypedElement::ColumnUtils::Argument> Arguments) override;
+	bool AddColumns(TypedElementRowHandle Row, TConstArrayView<const UScriptStruct*> Columns) override;
+	void RemoveColumns(TypedElementRowHandle Row, TConstArrayView<const UScriptStruct*> Columns) override;
+	bool AddRemoveColumns(TypedElementRowHandle Row, TConstArrayView<const UScriptStruct*> ColumnsToAdd,
+		TConstArrayView<const UScriptStruct*> ColumnsToRemove) override;
+	bool BatchAddRemoveColumns(TConstArrayView<TypedElementRowHandle> Rows,TConstArrayView<const UScriptStruct*> ColumnsToAdd,
+		TConstArrayView<const UScriptStruct*> ColumnsToRemove) override;
 	ColumnDataResult GetColumnData(TypedElementRowHandle Row, FTopLevelAssetPath ColumnName) override;
+	bool HasColumns(TypedElementRowHandle Row, TConstArrayView<const UScriptStruct*> ColumnTypes) const override;
+	bool HasColumns(TypedElementRowHandle Row, TConstArrayView<TWeakObjectPtr<const UScriptStruct>> ColumnTypes) const override;
 
 	void RegisterTickGroup(FName GroupName, EQueryTickPhase Phase, FName BeforeGroup, FName AfterGroup, bool bRequiresMainThread);
 	void UnregisterTickGroup(FName GroupName, EQueryTickPhase Phase);
@@ -101,6 +111,9 @@ private:
 		bool bRequiresMainThread{ false };
 	};
 	
+	/** Converts a set of column types into Mass specific fragment and tag bit sets. Returns true if any values were added. */
+	static bool ColumnsToBitSets(TConstArrayView<const UScriptStruct*> Columns, FMassFragmentBitSet& Fragments, FMassTagBitSet& Tags);
+
 	void PreparePhase(EQueryTickPhase Phase, float DeltaTime);
 	void FinalizePhase(EQueryTickPhase Phase, float DeltaTime);
 	void PhasePreOrPostAmble(EQueryTickPhase Phase, float DeltaTime, TArray<TypedElementQueryHandle>& Queries);
diff --git a/Engine/Plugins/Experimental/TypedElementsDataStorage/Source/TypedElementsDataStorage/Private/TypedElementDatabaseCompatibility.cpp b/Engine/Plugins/Experimental/TypedElementsDataStorage/Source/TypedElementsDataStorage/Private/TypedElementDatabaseCompatibility.cpp
--- a/Engine/Plugins/Experimental/TypedElementsDataStorage/Source/TypedElementsDataStorage/Private/TypedElementDatabaseCompatibility.cpp
+++ b/Engine/Plugins/Experimental/TypedElementsDataStorage/Source/TypedElementsDataStorage/Private/TypedElementDatabaseCompatibility.cpp
@@ -214,7 +214,7 @@ void UTypedElementDatabaseCompatibility::Tick()
 				// Note: Watch out for the performance of this, may end up doing a lot of row moves
 				for (TypedElementRowHandle Row : RowHandles)
 				{
-					Storage->AddTag(Row, FTypedElementSyncFromWorldTag::StaticStruct());
+					Storage->AddColumn<FTypedElementSyncFromWorldTag>(Row);
 				}
 			}
 		}
diff --git a/Engine/Plugins/Experimental/TypedElementsDataStorage/Source/TypedElementsDataStorageUI/Private/TypedElementCounterWidget.cpp b/Engine/Plugins/Experimental/TypedElementsDataStorage/Source/TypedElementsDataStorageUI/Private/TypedElementCounterWidget.cpp
--- a/Engine/Plugins/Experimental/TypedElementsDataStorage/Source/TypedElementsDataStorageUI/Private/TypedElementCounterWidget.cpp
+++ b/Engine/Plugins/Experimental/TypedElementsDataStorage/Source/TypedElementsDataStorageUI/Private/TypedElementCounterWidget.cpp
@@ -168,13 +168,14 @@ void FTypedElementCounterWidgetConstructor::AddColumns(
 	ITypedElementDataStorageInterface* DataStorage, TypedElementRowHandle Row, const TSharedPtr<SWidget>& Widget)
 {
 	Super::AddColumns(DataStorage, Row, Widget);
+	DataStorage->AddColumns<FTypedElementCounterWidgetColumn, FTypedElementU32IntValueCacheColumn>(Row);
 
-	FTypedElementCounterWidgetColumn* CounterColumn = DataStorage->AddOrGetColumn<FTypedElementCounterWidgetColumn>(Row);
+	FTypedElementCounterWidgetColumn* CounterColumn = DataStorage->GetColumn<FTypedElementCounterWidgetColumn>(Row);
 	checkf(CounterColumn, TEXT("Added a new FTypedElementCounterWidgetColumn to the Typed Elements Data Storage, but didn't get a valid pointer back."));
 	CounterColumn->LabelTextFormatter = LabelText;
 	CounterColumn->Query = Query;
 
-	FTypedElementU32IntValueCacheColumn* CacheColumn = DataStorage->AddOrGetColumn<FTypedElementU32IntValueCacheColumn>(Row);
+	FTypedElementU32IntValueCacheColumn* CacheColumn = DataStorage->GetColumn<FTypedElementU32IntValueCacheColumn>(Row);
 	checkf(CacheColumn, TEXT("Added a new FTypedElementUnsigned32BitIntValueCache to the Typed Elements Data Storage, but didn't get a valid pointer back."));
 	CacheColumn->Value = 0;
 }
diff --git a/Engine/Source/Runtime/TypedElementFramework/Private/Elements/Interfaces/TypedElementDataStorageUiInterface.cpp b/Engine/Source/Runtime/TypedElementFramework/Private/Elements/Interfaces/TypedElementDataStorageUiInterface.cpp
--- a/Engine/Source/Runtime/TypedElementFramework/Private/Elements/Interfaces/TypedElementDataStorageUiInterface.cpp
+++ b/Engine/Source/Runtime/TypedElementFramework/Private/Elements/Interfaces/TypedElementDataStorageUiInterface.cpp
@@ -37,8 +37,6 @@ void FTypedElementWidgetConstructor::ApplyArguments(TConstArrayView<TypedElement
 void FTypedElementWidgetConstructor::AddColumns(
 	ITypedElementDataStorageInterface* DataStorage, TypedElementRowHandle Row, const TSharedPtr<SWidget>& Widget)
 {
-	DataStorage->AddTag<FTypedElementSlateWidgetReferenceDeletesRowTag>(Row);
-	FTypedElementSlateWidgetReferenceColumn* StorageColumn = DataStorage->AddOrGetColumn<FTypedElementSlateWidgetReferenceColumn>(Row);
-	checkf(StorageColumn, TEXT("Expected Typed Element Data Storage to contain a FTypedElementSlateWidgetReferenceFragment for row %llu."), Row);
-	StorageColumn->Widget = Widget;
+	DataStorage->AddColumns<FTypedElementSlateWidgetReferenceColumn, FTypedElementSlateWidgetReferenceDeletesRowTag>(Row);
+	DataStorage->GetColumn<FTypedElementSlateWidgetReferenceColumn>(Row)->Widget = Widget;
 }
\ No newline at end of file
diff --git a/Engine/Source/Runtime/TypedElementFramework/Public/Elements/Interfaces/TypedElementDataStorageInterface.h b/Engine/Source/Runtime/TypedElementFramework/Public/Elements/Interfaces/TypedElementDataStorageInterface.h
--- a/Engine/Source/Runtime/TypedElementFramework/Public/Elements/Interfaces/TypedElementDataStorageInterface.h
+++ b/Engine/Source/Runtime/TypedElementFramework/Public/Elements/Interfaces/TypedElementDataStorageInterface.h
@@ -69,6 +69,19 @@ class TYPEDELEMENTFRAMEWORK_API ITypedElementDataStorageInterface
 	GENERATED_BODY()
 
 public:
+	/** 
+	 * Convenience structure that can be used to pass a lis of columns to functions that don't
+	 * have an dedicate templated version that takes a column list directly, for instance when
+	 * multiple column lists are used.
+	 */
+	template<typename... Columns>
+	struct TColumnTypeList
+	{
+		const UScriptStruct* ColumnTypes[sizeof...(Columns)] = { Columns::StaticStruct()... };
+
+		operator TConstArrayView<const UScriptStruct*>() { return ColumnTypes; }
+	};
+
 	/**
 	 * @section Table management
 	 * 
@@ -109,7 +122,7 @@ public:
 	virtual bool BatchAddRow(TypedElementTableHandle Table, int32 Count, TypedElementDataStorageCreationCallbackRef OnCreated) = 0;
 	virtual bool BatchAddRow(FName TableName, int32 Count, TypedElementDataStorageCreationCallbackRef OnCreated) = 0;
 	
-	/** Removes a previously added row. */
+	/** Removes a previously added row. If the row handle is invalid or already removed, nothing happens */
 	virtual void RemoveRow(TypedElementRowHandle Row) = 0;
 
 	
@@ -117,25 +130,59 @@ public:
 	 * @section Column management
 	 */
 
-	/** Adds a tag to a row or does nothing if already added. */
-	virtual void AddTag(TypedElementRowHandle Row, const UScriptStruct* TagType) = 0;
-	virtual void AddTag(TypedElementRowHandle Row, FTopLevelAssetPath TagName) = 0;
-	
-	/** Adds a new column to a row. If the column already exists it will be returned instead. */
+	/** Adds a column to a row or does nothing if already added. */
+	virtual bool AddColumn(TypedElementRowHandle Row, const UScriptStruct* ColumnType) = 0;
+	virtual bool AddColumn(TypedElementRowHandle Row, FTopLevelAssetPath ColumnName) = 0;
+	/**
+	 * Adds multiple columns from a row. This is typically more efficient than adding columns one 
+	 * at a time.
+	 */
+	virtual bool AddColumns(TypedElementRowHandle Row, TConstArrayView<const UScriptStruct*> Columns) = 0;
+
+	/** Removes a column from a row or does nothing if already removed. */
+	virtual void RemoveColumn(TypedElementRowHandle Row, const UScriptStruct* ColumnType) = 0;
+	virtual void RemoveColumn(TypedElementRowHandle Row, FTopLevelAssetPath ColumnName) = 0;
+	/**
+	 * Removes multiple columns from a row. This is typically more efficient than adding columns one
+	 * at a time.
+	 */
+	virtual void RemoveColumns(TypedElementRowHandle Row, TConstArrayView<const UScriptStruct*> Columns) = 0;
+	/** 
+	 * Adds and removes the provided column types from the provided row. This is typically more efficient 
+	 * than individually adding and removing columns as well as being faster than adding and removing
+	 * columns separately.
+	 */
+	virtual bool AddRemoveColumns(TypedElementRowHandle Row, TConstArrayView<const UScriptStruct*> ColumnsToAdd,
+		TConstArrayView<const UScriptStruct*> ColumnsToRemove) = 0;
+
+	/** Adds and removes the provided column types from the provided list of rows. */
+	virtual bool BatchAddRemoveColumns(
+		TConstArrayView<TypedElementRowHandle> Rows,
+		TConstArrayView<const UScriptStruct*> ColumnsToAdd,
+		TConstArrayView<const UScriptStruct*> ColumnsToRemove) = 0;
+
+	/**
+	 * Adds a new column to a row. If the column already exists it will be returned instead. If the colum couldn't
+	 * be added or the column type points to a tag an nullptr will be returned.
+	 */
 	virtual void* AddOrGetColumnData(TypedElementRowHandle Row, const UScriptStruct* ColumnType) = 0;
 	virtual ColumnDataResult AddOrGetColumnData(TypedElementRowHandle Row, FTopLevelAssetPath ColumnName) = 0;
-	
 	/**
 	 * Sets the data of a column using the provided argument bag. This is only meant for simple initialization for 
 	 * fragments that use UPROPERTY to expose properties. For complex initialization or when the fragment type is known
-	 * it's recommende to use call that work directly on the type for better performance and a wider range of configuration options.
+	 * it's recommended to use calls that work directly on the type for better performance and a wider range of configuration options.
+	 * If the column couldn't be created or the column name points to a tag, then the result will contain only nullptrs.
 	 */
 	virtual ColumnDataResult AddOrGetColumnData(TypedElementRowHandle Row, FTopLevelAssetPath ColumnName,
 		TConstArrayView<TypedElement::ColumnUtils::Argument> Arguments) = 0;
 	
-	/** Retrieves a pointer to the column of the given row or a nullptr if not found. */
+	/** Retrieves a pointer to the column of the given row or a nullptr if not found or if the column type is a tag. */
 	virtual void* GetColumnData(TypedElementRowHandle Row, const UScriptStruct* ColumnType) = 0;
 	virtual ColumnDataResult GetColumnData(TypedElementRowHandle Row, FTopLevelAssetPath ColumnName) = 0;
+
+	/** Determines if the provided row contains the collection of columsn and tags. */
+	virtual bool HasColumns(TypedElementRowHandle Row, TConstArrayView<const UScriptStruct*> ColumnTypes) const = 0;
+	virtual bool HasColumns(TypedElementRowHandle Row, TConstArrayView<TWeakObjectPtr<const UScriptStruct>> ColumnTypes) const = 0;
 	
 
 	/**
@@ -418,7 +465,13 @@ public:
 		/** If true, this query only has simple operations and is guaranteed to be executed fully and at optimal performance. */
 		bool bSimpleQuery{ false };
 	};
+	/** 
+	 * Registers a query with the data storage. The description is processed into an internal format and may be changed. If no valid
+	 * could be created an invalid query handle will be returned. It's recommended to use the Query Builder for a more convenient
+	 * and safer construction of a query.
+	 */
 	virtual TypedElementQueryHandle RegisterQuery(FQueryDescription&& Query) = 0;
+	/** Removes a previous registered. If the query handle is invalid or the query has already been deleted nothing will happen. */
 	virtual void UnregisterQuery(TypedElementQueryHandle Query) = 0;
 	/** Returns the description of a previously registered query. If the query no longer exists an empty description will be returned. */
 	virtual const FQueryDescription& GetQueryDescription(TypedElementQueryHandle Query) const = 0;
@@ -450,8 +503,14 @@ public:
 		uint32 Count{ 0 }; /** The number of rows were processed. */
 		ECompletion Completed{ ECompletion::Unavailable };
 	};
+	/** Directly runs a query. If the query handle is invalid or has been deleted nothing will happen. */
 	virtual FQueryResult RunQuery(TypedElementQueryHandle Query) = 0;
-	virtual FQueryResult RunQuery(TypedElementQueryHandle Query, DirectQueryCallbackRef Callback) = 0;
+	/**
+	 * Directly runs a query. The callback will be called for batches of matching rows. During a single call to RunQuery the callback
+	 * may be called multiple times. If the query handle is invalid or has been deleted nothing happens and the callback won't be called
+	 */
+	virtual FQueryResult RunQuery(TypedElementQueryHandle Query, 
+		DirectQueryCallbackRef Callback) = 0;
 	
 	/**
 	 * @section Misc
@@ -482,10 +541,28 @@ public:
 	 */
 	
 	
-	/** Adds a tag to a row or does nothing if already added. */
-	template<typename TagType>
-	void AddTag(TypedElementRowHandle Row);
+	/** Adds a column to a row or does nothing if already added. */
+	template<typename Column>
+	bool AddColumn(TypedElementRowHandle Row);
+
+	/** Removes a tag from a row or does nothing if already removed. */
+	template<typename Column>
+	void RemoveColumn(TypedElementRowHandle Row);
 	
+	/**
+	 * Adds multiple columns from a row. This is typically more efficient than adding columns one
+	 * at a time.
+	 */
+	template<typename... Columns>
+	void AddColumns(TypedElementRowHandle Row);
+
+	/**
+	 * Removes multiple columns from a row. This is typically more efficient than adding columns one
+	 * at a time.
+	 */
+	template<typename... Columns>
+	void RemoveColumns(TypedElementRowHandle Row);
+
 	/**
 	 * Returns a pointer to the column of the given row or creates a new one if not found. Optionally arguments can be provided
 	 * to update or initialize the column's data.
@@ -497,6 +574,9 @@ public:
 	template<typename ColumnType>
 	ColumnType* GetColumn(TypedElementRowHandle Row);
 
+	template<typename... ColumnTypes>
+	bool HasColumns(TypedElementRowHandle Row) const;
+
 	/** Returns a pointer to the registered external system if found, otherwise null. */
 	template<typename SystemType>
 	SystemType* GetExternalSystem();
@@ -508,10 +588,28 @@ public:
 
 ENUM_CLASS_FLAGS(ITypedElementDataStorageInterface::EQueryDependencyFlags);
 
-template<typename TagType>
-void ITypedElementDataStorageInterface::AddTag(TypedElementRowHandle Row)
+template<typename Column>
+bool ITypedElementDataStorageInterface::AddColumn(TypedElementRowHandle Row)
+{
+	return AddColumn(Row, Column::StaticStruct());
+}
+
+template<typename Column>
+void ITypedElementDataStorageInterface::RemoveColumn(TypedElementRowHandle Row)
+{
+	RemoveColumn(Row, Column::StaticStruct());
+}
+
+template<typename... Columns>
+void ITypedElementDataStorageInterface::AddColumns(TypedElementRowHandle Row)
+{
+	AddColumns(Row, { Columns::StaticStruct()...});
+}
+
+template<typename... Columns>
+void ITypedElementDataStorageInterface::RemoveColumns(TypedElementRowHandle Row)
 {
-	AddTag(Row, TagType::StaticStruct());
+	RemoveColumns(Row, { Columns::StaticStruct()...});
 }
 
 template<typename ColumnType, typename... Args>
@@ -528,24 +626,18 @@ ColumnType* ITypedElementDataStorageInterface::AddOrGetColumn(TypedElementRowHan
 	return Result;
 }
 
-template<typename Column>
-const Column* ITypedElementDataStorageInterface::ICommonQueryContext::GetColumn() const
-{
-	return reinterpret_cast<const Column*>(GetColumn(Column::StaticStruct()));
-}
-
-template<typename Column>
-Column* ITypedElementDataStorageInterface::ICommonQueryContext::GetMutableColumn()
-{
-	return reinterpret_cast<Column*>(GetMutableColumn(Column::StaticStruct()));
-}
-
 template<typename ColumnType>
 ColumnType* ITypedElementDataStorageInterface::GetColumn(TypedElementRowHandle Row)
 {
 	return reinterpret_cast<ColumnType*>(GetColumnData(Row, ColumnType::StaticStruct()));
 }
 
+template<typename... ColumnType>
+bool ITypedElementDataStorageInterface::HasColumns(TypedElementRowHandle Row) const
+{
+	return HasColumns(Row, TConstArrayView<const UScriptStruct*>({ ColumnType::StaticStruct()... }));
+}
+
 template<typename SystemType>
 SystemType* ITypedElementDataStorageInterface::GetExternalSystem()
 {
@@ -553,6 +645,23 @@ SystemType* ITypedElementDataStorageInterface::GetExternalSystem()
 }
 
 
+
+//
+// ITypedElementDataStorageInterface::IQueryContext
+//
+
+template<typename Column>
+const Column* ITypedElementDataStorageInterface::ICommonQueryContext::GetColumn() const
+{
+	return reinterpret_cast<const Column*>(GetColumn(Column::StaticStruct()));
+}
+
+template<typename Column>
+Column* ITypedElementDataStorageInterface::ICommonQueryContext::GetMutableColumn()
+{
+	return reinterpret_cast<Column*>(GetMutableColumn(Column::StaticStruct()));
+}
+
 template<typename... Columns>
 void ITypedElementDataStorageInterface::IQueryContext::AddColumns(TypedElementRowHandle Row)
 {
