Extended the Slate integration with TEDS.

The following changes were made:
- Matching columns to widgets can now be done in two additional ways besides longest match:
    - Exact match - the provided columns need to exactly match the columns associated with the widget.
    - Single match - each column is matched individually and they're not combined.
- Updated how columns are added to a widget row. Instead of the constructor being responsible for adding them, widget constructors now provide a list of columns to add and it'll be up to the caller to add them before calling the constructor. This allows callers to fill out the columns should they be needed by the widget constructor. This also reduces the risk of widget not batch adding columns.
- Changed the use of `GetWeakPtrTypeHash` to using the actual pointer instead. This is because `GetWeakPtrTypeHash` doesn't resolve to the templated type but always returns the same hash.
- Widget constructors can now indicate whether or not they can be reused. This allows avoiding creating temporary instances in some cases. The default is false.

diff --git a/Engine/Plugins/Experimental/TypedElementsDataStorage/Source/TypedElementsDataStorage/Private/TypedElementDatabaseUI.cpp b/Engine/Plugins/Experimental/TypedElementsDataStorage/Source/TypedElementsDataStorage/Private/TypedElementDatabaseUI.cpp
--- a/Engine/Plugins/Experimental/TypedElementsDataStorage/Source/TypedElementsDataStorage/Private/TypedElementDatabaseUI.cpp
+++ b/Engine/Plugins/Experimental/TypedElementsDataStorage/Source/TypedElementsDataStorage/Private/TypedElementDatabaseUI.cpp
@@ -120,10 +120,16 @@ bool UTypedElementDatabaseUi::RegisterWidgetFactory(
 				}
 				return true;
 			case ITypedElementDataStorageUiInterface::EPurposeType::UniqueByNameAndColumn:
-				PrepareColumnsList(Columns);
-				PurposeInfo->Factories.Emplace(Constructor, MoveTemp(Columns));
-				PurposeInfo->bIsSorted = false;
-				return true;
+				if (PrepareColumnsList(Columns))
+				{
+					PurposeInfo->Factories.Emplace(Constructor, MoveTemp(Columns));
+					PurposeInfo->bIsSorted = false;
+					return true;
+				}
+				else
+				{
+					return false;
+				}
 			default:
 				checkf(false, TEXT("Unexpected ITypedElementDataStorageUiInterface::EPurposeType found provided when registering widget factory."));
 				return false;
@@ -247,11 +253,10 @@ void UTypedElementDatabaseUi::CreateWidgetConstructors(FName Purpose,
 	}
 }
 
-void UTypedElementDatabaseUi::CreateWidgetConstructors(FName Purpose, TArray<TWeakObjectPtr<const UScriptStruct>>& Columns,
-	TConstArrayView<TypedElement::ColumnUtils::Argument> Arguments, const WidgetConstructorCallback& Callback)
+void UTypedElementDatabaseUi::CreateWidgetConstructors(FName Purpose, EMatchApproach MatchApproach, 
+	TArray<TWeakObjectPtr<const UScriptStruct>>& Columns, TConstArrayView<TypedElement::ColumnUtils::Argument> Arguments, 
+	const WidgetConstructorCallback& Callback)
 {
-	using ColumnView = TConstArrayView<TWeakObjectPtr<const UScriptStruct>>;
-	
 	if (FPurposeInfo* PurposeInfo = WidgetPurposes.Find(Purpose))
 	{
 		// Sort so searching can be done in a single pass. This would also allow for binary searching, but the number of columns
@@ -260,12 +265,12 @@ void UTypedElementDatabaseUi::CreateWidgetConstructors(FName Purpose, TArray<TWe
 		Columns.Sort(
 			[](const TWeakObjectPtr<const UScriptStruct>& Lhs, const TWeakObjectPtr<const UScriptStruct>& Rhs)
 			{
-				return Lhs.GetWeakPtrTypeHash() < Rhs.GetWeakPtrTypeHash();
+				return Lhs.Get() < Rhs.Get();
 			});
 
 		if (!PurposeInfo->bIsSorted)
 		{
-			// This is the only call that requires the array of factories to be sorted from from largest to smallest number
+			// This is the only call that requires the array of factories to be sorted from largest to smallest number
 			// of columns, so lazily sort only when needed.
 			PurposeInfo->Factories.StableSort(
 				[](const FWidgetFactory& Lhs, const FWidgetFactory& Rhs)
@@ -276,11 +281,11 @@ void UTypedElementDatabaseUi::CreateWidgetConstructors(FName Purpose, TArray<TWe
 					{
 						for (int32 Index = 0; Index < LeftSize; ++Index)
 						{
-							uint32 LeftHash = Lhs.Columns[Index].GetWeakPtrTypeHash();
-							uint32 RightHash = Lhs.Columns[Index].GetWeakPtrTypeHash();
-							if (LeftHash != RightHash)
+							const UScriptStruct* LhsStruct = Lhs.Columns[Index].Get();
+							const UScriptStruct* RhsStruct = Rhs.Columns[Index].Get();
+							if (LhsStruct != RhsStruct)
 							{
-								return LeftHash < RightHash;
+								return LhsStruct < RhsStruct;
 							}
 						}
 						return true;
@@ -293,51 +298,20 @@ void UTypedElementDatabaseUi::CreateWidgetConstructors(FName Purpose, TArray<TWe
 			PurposeInfo->bIsSorted = true;
 		}
 
-		for (auto FactoryIt = PurposeInfo->Factories.CreateConstIterator();  FactoryIt && !Columns.IsEmpty();  ++FactoryIt)
+		switch (MatchApproach)
 		{
-			int32 MatchingIndex = INDEX_NONE;
-			if (FactoryIt->Columns.Num() > Columns.Num())
-			{
-				// There are more columns required for this factory than there are in the requested columns list so skip this
-				// factory.
-				continue;
-			}
-			
-			ColumnView ColumnsRange = Columns;
-			for (auto ColumnIt = FactoryIt->Columns.CreateConstIterator(); ColumnIt && !ColumnsRange.IsEmpty(); ++ColumnIt)
-			{
-				int32 Index = INDEX_NONE;
-				if (ColumnsRange.Find(*ColumnIt, Index))
-				{
-					MatchingIndex = MatchingIndex == INDEX_NONE ? Index : MatchingIndex;
-					ColumnsRange = ColumnsRange.RightChop(Index + 1);
-				}
-				else
-				{
-					// Reset as there is at least one column that didn't match.
-					MatchingIndex = INDEX_NONE;
-					break;
-				}
-			}
-
-			// This assumes that the column arrays are sorted from longest to shortest.
-			if (MatchingIndex != INDEX_NONE)
-			{
-				// Delete in reverse order to reduce the number of objects that need to be moved up as the order needs to remain stable.
-				ColumnsRange = ColumnView(Columns);
-				for (auto It = FactoryIt->Columns.rbegin(); It != FactoryIt->Columns.rend(); ++It)
-				{
-					int32 Index = INDEX_NONE;
-					ensureMsgf(ColumnsRange.Find(*It, Index), TEXT("A previous found matching column can't be found in the original array."));
-					Columns.RemoveAt(Index);
-					ColumnsRange = ColumnsRange.Left(Index);
-				}
-
-				if (!CreateSingleWidgetConstructor(FactoryIt->Constructor, Arguments, FactoryIt->Columns, Callback))
-				{
-					return;
-				}
-			}
+		case EMatchApproach::LongestMatch:
+			CreateWidgetConstructors_LongestMatch(PurposeInfo->Factories, Columns, Arguments, Callback);
+			break;
+		case EMatchApproach::ExactMatch:
+			CreateWidgetConstructors_ExactMatch(PurposeInfo->Factories, Columns, Arguments, Callback);
+			break;
+		case EMatchApproach::SingleMatch:
+			CreateWidgetConstructors_SingleMatch(PurposeInfo->Factories, Columns, Arguments, Callback);
+			break;
+		default:
+			checkf(false, TEXT("Unsupported match type (%i) for CreateWidgetConstructors."), 
+				static_cast<std::underlying_type_t<EMatchApproach>>(MatchApproach));
 		}
 	}
 }
@@ -424,7 +398,7 @@ void UTypedElementDatabaseUi::CreateWidgetInstanceFromInstance(
 	TConstArrayView<TypedElement::ColumnUtils::Argument> Arguments,
 	const WidgetCreatedCallback& ConstructionCallback)
 {
-	if (Arguments.IsEmpty())
+	if (SourceConstructor->CanBeReused())
 	{
 		CreateWidgetInstance(*SourceConstructor, Arguments, ConstructionCallback);
 	}
@@ -454,6 +428,7 @@ void UTypedElementDatabaseUi::CreateWidgetInstance(
 	const WidgetCreatedCallback& ConstructionCallback)
 {
 	TypedElementRowHandle Row = Storage->AddRow(WidgetTable);
+	Storage->AddColumns(Row, Constructor.GetAdditionalColumnsList());
 	TSharedPtr<SWidget> Widget = Constructor.Construct(Row, Storage, this, Arguments);
 	if (Widget)
 	{
@@ -468,7 +443,7 @@ void UTypedElementDatabaseUi::CreateWidgetInstance(
 TSharedPtr<SWidget> UTypedElementDatabaseUi::ConstructWidget(TypedElementRowHandle Row, FTypedElementWidgetConstructor& Constructor,
 	TConstArrayView<TypedElement::ColumnUtils::Argument> Arguments)
 {
-	if (Arguments.IsEmpty())
+	if (Constructor.CanBeReused())
 	{
 		return Constructor.Construct(Row, Storage, this, Arguments);
 	}
@@ -508,21 +483,28 @@ void UTypedElementDatabaseUi::CreateStandardArchetypes()
 		{
 			FTypedElementSlateWidgetReferenceColumn::StaticStruct(),
 			FTypedElementSlateWidgetReferenceDeletesRowTag::StaticStruct()
-		}), FName("Editor_WidgetTable"));
+		}), FName(TEXT("Editor_WidgetTable")));
 }
 
-void UTypedElementDatabaseUi::PrepareColumnsList(TArray<TWeakObjectPtr<const UScriptStruct>>& Columns)
+bool UTypedElementDatabaseUi::PrepareColumnsList(TArray<TWeakObjectPtr<const UScriptStruct>>& Columns)
 {
 	Columns.Sort(
 		[](const TWeakObjectPtr<const UScriptStruct>& Lhs, const TWeakObjectPtr<const UScriptStruct>& Rhs)
 		{
-			return Lhs.GetWeakPtrTypeHash() < Rhs.GetWeakPtrTypeHash();
+			return Lhs.Get() < Rhs.Get();
 		});
 
+	// Check if there are any columns to register and whether there are any invalid columns. After sorting
+	// invalid columns will always be at the start of the array.
+	if (Columns.IsEmpty() || Columns[0] == nullptr)
+	{
+		return false;
+	}
+
 	// Remove duplicates.
 	for (int32 Index = 1; Index < Columns.Num(); )
 	{
-		if (Columns[Index - 1].GetWeakPtrTypeHash() != Columns[Index].GetWeakPtrTypeHash())
+		if (Columns[Index - 1] != Columns[Index])
 		{
 			++Index;
 		}
@@ -531,6 +513,140 @@ void UTypedElementDatabaseUi::PrepareColumnsList(TArray<TWeakObjectPtr<const USc
 			Columns.RemoveAt(Index);
 		}
 	}
+	return true;
+}
+
+void UTypedElementDatabaseUi::CreateWidgetConstructors_LongestMatch(const TArray<FWidgetFactory>& WidgetFactories,
+	TArray<TWeakObjectPtr<const UScriptStruct>>& Columns, TConstArrayView<TypedElement::ColumnUtils::Argument> Arguments,
+	const WidgetConstructorCallback& Callback)
+{
+	using ColumnView = TConstArrayView<TWeakObjectPtr<const UScriptStruct>>;
+
+	for (auto FactoryIt = WidgetFactories.CreateConstIterator(); FactoryIt && !Columns.IsEmpty(); ++FactoryIt)
+	{
+		int32 MatchingIndex = INDEX_NONE;
+		if (FactoryIt->Columns.Num() > Columns.Num())
+		{
+			// There are more columns required for this factory than there are in the requested columns list so skip this
+			// factory.
+			continue;
+		}
+
+		ColumnView ColumnsRange = Columns;
+		for (auto ColumnIt = FactoryIt->Columns.CreateConstIterator(); ColumnIt && !ColumnsRange.IsEmpty(); ++ColumnIt)
+		{
+			int32 Index = INDEX_NONE;
+			if (ColumnsRange.Find(*ColumnIt, Index))
+			{
+				MatchingIndex = MatchingIndex == INDEX_NONE ? Index : MatchingIndex;
+				ColumnsRange = ColumnsRange.RightChop(Index + 1);
+			}
+			else
+			{
+				// Reset as there is at least one column that didn't match.
+				MatchingIndex = INDEX_NONE;
+				break;
+			}
+		}
+
+		// This assumes that the column arrays are sorted from longest to shortest.
+		if (MatchingIndex != INDEX_NONE)
+		{
+			// Delete in reverse order to reduce the number of objects that need to be moved up as the order needs to remain stable.
+			ColumnsRange = ColumnView(Columns);
+			for (auto It = FactoryIt->Columns.rbegin(); It != FactoryIt->Columns.rend(); ++It)
+			{
+				int32 Index = INDEX_NONE;
+				ensureMsgf(ColumnsRange.Find(*It, Index), TEXT("A previous found matching column can't be found in the original array."));
+				Columns.RemoveAt(Index);
+				ColumnsRange = ColumnsRange.Left(Index);
+			}
+
+			if (!CreateSingleWidgetConstructor(FactoryIt->Constructor, Arguments, FactoryIt->Columns, Callback))
+			{
+				return;
+			}
+		}
+	}
+}
+
+void UTypedElementDatabaseUi::CreateWidgetConstructors_ExactMatch(const TArray<FWidgetFactory>& WidgetFactories,
+	TArray<TWeakObjectPtr<const UScriptStruct>>& Columns, TConstArrayView<TypedElement::ColumnUtils::Argument> Arguments,
+	const WidgetConstructorCallback& Callback)
+{
+	int32 ColumnCount = Columns.Num();
+	for (const FWidgetFactory& Factory : WidgetFactories)
+	{
+		if (Factory.Columns.Num() != ColumnCount)
+		{
+			continue;
+		}
+
+		bool bFoundMatch = true;
+		TWeakObjectPtr<const UScriptStruct>* ColumnIt = Columns.GetData();
+		const TWeakObjectPtr<const UScriptStruct>* FactoryColumnsIt = Factory.Columns.GetData();
+		for (int32 Index = 0; Index < ColumnCount; ++Index)
+		{
+			if (ColumnIt != FactoryColumnsIt)
+			{
+				bFoundMatch = false;
+				break;
+			}
+			++ColumnIt;
+			++FactoryColumnsIt;
+		}
+		if (bFoundMatch)
+		{
+			Columns.Reset();
+			CreateSingleWidgetConstructor(Factory.Constructor, Arguments, Factory.Columns, Callback);
+			return;
+		}
+	}
+}
+
+void UTypedElementDatabaseUi::CreateWidgetConstructors_SingleMatch(const TArray<FWidgetFactory>& WidgetFactories,
+	TArray<TWeakObjectPtr<const UScriptStruct>>& Columns, TConstArrayView<TypedElement::ColumnUtils::Argument> Arguments,
+	const WidgetConstructorCallback& Callback)
+{
+	auto FactoryIt = WidgetFactories.rbegin();
+	auto FactoryEnd = WidgetFactories.rend();
+
+	// Start from the back as the widgets with lower counts will be last.
+	int32 ColumnIndex = Columns.Num() - 1;
+
+	auto ColumnEnd = Columns.rend();
+	for (auto ColumnIt = Columns.rbegin(); ColumnIt != ColumnEnd; ++ColumnIt, --ColumnIndex)
+	{
+		const UScriptStruct* ColumnType = (*ColumnIt).Get();
+		for (; FactoryIt != FactoryEnd; ++FactoryIt)
+		{
+			int32 ColumnCount = (*FactoryIt).Columns.Num();
+			if (ColumnCount > 1)
+			{
+				// Moved passed the point where factories only have a single column.
+				return;
+			}
+			else if (ColumnCount == 0)
+			{
+				// Need to move further to find factories with exactly one column.
+				continue;
+			}
+
+			const UScriptStruct* FactoryType = (*FactoryIt).Columns[0].Get();
+			if (FactoryType == ColumnType)
+			{
+				Columns.RemoveAt(ColumnIndex);
+				CreateSingleWidgetConstructor((*FactoryIt).Constructor, Arguments, (*FactoryIt).Columns, Callback);
+				// Match was found so move on to the next column in the column.
+				break;
+			}
+			if (FactoryType < ColumnType)
+			{
+				// The current column isn't found in the factories, so move on to the next column.
+				break;
+			}
+		}
+	}
 }
 
 
diff --git a/Engine/Plugins/Experimental/TypedElementsDataStorage/Source/TypedElementsDataStorage/Private/TypedElementDatabaseUI.h b/Engine/Plugins/Experimental/TypedElementsDataStorage/Source/TypedElementsDataStorage/Private/TypedElementDatabaseUI.h
--- a/Engine/Plugins/Experimental/TypedElementsDataStorage/Source/TypedElementsDataStorage/Private/TypedElementDatabaseUI.h
+++ b/Engine/Plugins/Experimental/TypedElementsDataStorage/Source/TypedElementsDataStorage/Private/TypedElementDatabaseUI.h
@@ -41,7 +41,7 @@ public:
 
 	void CreateWidgetConstructors(FName Purpose,
 		TConstArrayView<TypedElement::ColumnUtils::Argument> Arguments, const WidgetConstructorCallback& Callback) override;
-	void CreateWidgetConstructors(FName Purpose, TArray<TWeakObjectPtr<const UScriptStruct>>& Columns,
+	void CreateWidgetConstructors(FName Purpose, EMatchApproach MatchApproach, TArray<TWeakObjectPtr<const UScriptStruct>>& Columns,
 		TConstArrayView<TypedElement::ColumnUtils::Argument> Arguments, const WidgetConstructorCallback& Callback) override;
 
 	void ConstructWidgets(FName Purpose, TConstArrayView<TypedElement::ColumnUtils::Argument> Arguments,
@@ -97,7 +97,23 @@ private:
 		TConstArrayView<TypedElement::ColumnUtils::Argument> Arguments,
 		const WidgetCreatedCallback& ConstructionCallback);
 
-	static void PrepareColumnsList(TArray<TWeakObjectPtr<const UScriptStruct>>& Columns);
+	static bool PrepareColumnsList(TArray<TWeakObjectPtr<const UScriptStruct>>& Columns);
+
+	void CreateWidgetConstructors_LongestMatch(
+		const TArray<FWidgetFactory>& WidgetFactories, 
+		TArray<TWeakObjectPtr<const UScriptStruct>>& Columns, 
+		TConstArrayView<TypedElement::ColumnUtils::Argument> Arguments, 
+		const WidgetConstructorCallback& Callback);
+	void CreateWidgetConstructors_ExactMatch(
+		const TArray<FWidgetFactory>& WidgetFactories,
+		TArray<TWeakObjectPtr<const UScriptStruct>>& Columns,
+		TConstArrayView<TypedElement::ColumnUtils::Argument> Arguments,
+		const WidgetConstructorCallback& Callback);
+	void CreateWidgetConstructors_SingleMatch(
+		const TArray<FWidgetFactory>& WidgetFactories,
+		TArray<TWeakObjectPtr<const UScriptStruct>>& Columns,
+		TConstArrayView<TypedElement::ColumnUtils::Argument> Arguments,
+		const WidgetConstructorCallback& Callback);
 
 	TypedElementTableHandle WidgetTable{ TypedElementInvalidTableHandle };
 	
diff --git a/Engine/Plugins/Experimental/TypedElementsDataStorage/Source/TypedElementsDataStorageUI/Private/TypedElementCounterWidget.cpp b/Engine/Plugins/Experimental/TypedElementsDataStorage/Source/TypedElementsDataStorageUI/Private/TypedElementCounterWidget.cpp
--- a/Engine/Plugins/Experimental/TypedElementsDataStorage/Source/TypedElementsDataStorageUI/Private/TypedElementCounterWidget.cpp
+++ b/Engine/Plugins/Experimental/TypedElementsDataStorage/Source/TypedElementsDataStorageUI/Private/TypedElementCounterWidget.cpp
@@ -145,7 +145,6 @@ void UTypedElementCounterWidgetFactory::SetupMainWindowIntegrations(TSharedPtr<S
 
 
 
-
 //
 // FTypedElementCounterWidgetConstructor
 //
@@ -155,6 +154,17 @@ FTypedElementCounterWidgetConstructor::FTypedElementCounterWidgetConstructor()
 {
 }
 
+TConstArrayView<const UScriptStruct*> FTypedElementCounterWidgetConstructor::GetAdditionalColumnsList() const
+{
+	static TTypedElementColumnTypeList<FTypedElementCounterWidgetColumn, FTypedElementU32IntValueCacheColumn> Columns;
+	return Columns;
+}
+
+bool FTypedElementCounterWidgetConstructor::CanBeReused() const
+{
+	return true;
+}
+
 TSharedPtr<SWidget> FTypedElementCounterWidgetConstructor::CreateWidget()
 {
 	return SNew(STextBlock)
@@ -164,12 +174,8 @@ TSharedPtr<SWidget> FTypedElementCounterWidgetConstructor::CreateWidget()
 		.Justification(ETextJustify::Center);
 }
 
-void FTypedElementCounterWidgetConstructor::AddColumns(
-	ITypedElementDataStorageInterface* DataStorage, TypedElementRowHandle Row, const TSharedPtr<SWidget>& Widget)
+bool FTypedElementCounterWidgetConstructor::SetColumns(ITypedElementDataStorageInterface* DataStorage, TypedElementRowHandle Row)
 {
-	Super::AddColumns(DataStorage, Row, Widget);
-	DataStorage->AddColumns<FTypedElementCounterWidgetColumn, FTypedElementU32IntValueCacheColumn>(Row);
-
 	FTypedElementCounterWidgetColumn* CounterColumn = DataStorage->GetColumn<FTypedElementCounterWidgetColumn>(Row);
 	checkf(CounterColumn, TEXT("Added a new FTypedElementCounterWidgetColumn to the Typed Elements Data Storage, but didn't get a valid pointer back."));
 	CounterColumn->LabelTextFormatter = LabelText;
@@ -178,6 +184,8 @@ void FTypedElementCounterWidgetConstructor::AddColumns(
 	FTypedElementU32IntValueCacheColumn* CacheColumn = DataStorage->GetColumn<FTypedElementU32IntValueCacheColumn>(Row);
 	checkf(CacheColumn, TEXT("Added a new FTypedElementUnsigned32BitIntValueCache to the Typed Elements Data Storage, but didn't get a valid pointer back."));
 	CacheColumn->Value = 0;
+
+	return true;
 }
 
 #undef LOCTEXT_NAMESPACE
\ No newline at end of file
diff --git a/Engine/Plugins/Experimental/TypedElementsDataStorage/Source/TypedElementsDataStorageUI/Public/TypedElementCounterWidget.h b/Engine/Plugins/Experimental/TypedElementsDataStorage/Source/TypedElementsDataStorageUI/Public/TypedElementCounterWidget.h
--- a/Engine/Plugins/Experimental/TypedElementsDataStorage/Source/TypedElementsDataStorageUI/Public/TypedElementCounterWidget.h
+++ b/Engine/Plugins/Experimental/TypedElementsDataStorage/Source/TypedElementsDataStorageUI/Public/TypedElementCounterWidget.h
@@ -47,13 +47,16 @@ public:
 	FTypedElementCounterWidgetConstructor();
 	~FTypedElementCounterWidgetConstructor() override = default;
 
+	TConstArrayView<const UScriptStruct*> GetAdditionalColumnsList() const override;
+	bool CanBeReused() const override;
+
 	FText ToolTipText{ NSLOCTEXT("TypedElementUI_CounterWidget", "Tooltip", "Shows the total number found in the editor.") };
 	FText LabelText{ NSLOCTEXT("TypedElementUI_CounterWidget", "Label", "Counted") };
 	TypedElementQueryHandle Query;
 
 protected:
 	TSharedPtr<SWidget> CreateWidget() override;
-	void AddColumns(ITypedElementDataStorageInterface* DataStorage, TypedElementRowHandle Row, const TSharedPtr<SWidget>& Widget) override;
+	bool SetColumns(ITypedElementDataStorageInterface* DataStorage, TypedElementRowHandle Row) override;
 };
 
 USTRUCT(meta = (DisplayName = "Counter widget"))
diff --git a/Engine/Source/Runtime/TypedElementFramework/Private/Elements/Interfaces/TypedElementDataStorageUiInterface.cpp b/Engine/Source/Runtime/TypedElementFramework/Private/Elements/Interfaces/TypedElementDataStorageUiInterface.cpp
--- a/Engine/Source/Runtime/TypedElementFramework/Private/Elements/Interfaces/TypedElementDataStorageUiInterface.cpp
+++ b/Engine/Source/Runtime/TypedElementFramework/Private/Elements/Interfaces/TypedElementDataStorageUiInterface.cpp
@@ -4,7 +4,7 @@
 #include "Elements/Framework/TypedElementColumnUtils.h"
 #include "Elements/Columns/TypedElementSlateWidgetColumns.h"
 
-FTypedElementWidgetConstructor::FTypedElementWidgetConstructor(UScriptStruct* InTypeInfo)
+FTypedElementWidgetConstructor::FTypedElementWidgetConstructor(const UScriptStruct* InTypeInfo)
 	: TypeInfo(InTypeInfo)
 {
 }
@@ -14,29 +14,56 @@ const UScriptStruct* FTypedElementWidgetConstructor::GetTypeInfo() const
 	return TypeInfo;
 }
 
+TConstArrayView<const UScriptStruct*> FTypedElementWidgetConstructor::GetAdditionalColumnsList() const
+{
+	return {};
+}
+
+bool FTypedElementWidgetConstructor::CanBeReused() const
+{
+	return false;
+}
+
 TSharedPtr<SWidget> FTypedElementWidgetConstructor::Construct(
 	TypedElementRowHandle Row,
 	ITypedElementDataStorageInterface* DataStorage,
 	ITypedElementDataStorageUiInterface* DataStorageUi,
 	TConstArrayView<TypedElement::ColumnUtils::Argument> Arguments)
 {
-	ApplyArguments(Arguments);
-	TSharedPtr<SWidget> Widget = CreateWidget();
-	if (Widget)
+	if (ApplyArguments(Arguments))
 	{
-		AddColumns(DataStorage, Row, Widget);
+		TSharedPtr<SWidget> Widget = CreateWidget();
+		if (Widget)
+		{
+			DataStorage->GetColumn<FTypedElementSlateWidgetReferenceColumn>(Row)->Widget = Widget;
+			if (SetColumns(DataStorage, Row))
+			{
+				if (FinalizeWidget(DataStorage, DataStorageUi, Row, Widget))
+				{
+					return Widget;
+				}
+			}
+		}
 	}
-	return Widget;
+	return nullptr;
 }
 
-void FTypedElementWidgetConstructor::ApplyArguments(TConstArrayView<TypedElement::ColumnUtils::Argument> Arguments)
+bool FTypedElementWidgetConstructor::ApplyArguments(TConstArrayView<TypedElement::ColumnUtils::Argument> Arguments)
 {
 	SetColumnValues(*this, Arguments);
+	return true;
 }
 
-void FTypedElementWidgetConstructor::AddColumns(
-	ITypedElementDataStorageInterface* DataStorage, TypedElementRowHandle Row, const TSharedPtr<SWidget>& Widget)
+bool FTypedElementWidgetConstructor::SetColumns(ITypedElementDataStorageInterface* DataStorage, TypedElementRowHandle Row)
+{
+	return true;
+}
+
+bool FTypedElementWidgetConstructor::FinalizeWidget(
+	ITypedElementDataStorageInterface* DataStorage,
+	ITypedElementDataStorageUiInterface* DataStorageUi,
+	TypedElementRowHandle Row,
+	const TSharedPtr<SWidget>& Widget)
 {
-	DataStorage->AddColumns<FTypedElementSlateWidgetReferenceColumn, FTypedElementSlateWidgetReferenceDeletesRowTag>(Row);
-	DataStorage->GetColumn<FTypedElementSlateWidgetReferenceColumn>(Row)->Widget = Widget;
+	return true;
 }
\ No newline at end of file
diff --git a/Engine/Source/Runtime/TypedElementFramework/Public/Elements/Interfaces/TypedElementDataStorageInterface.h b/Engine/Source/Runtime/TypedElementFramework/Public/Elements/Interfaces/TypedElementDataStorageInterface.h
--- a/Engine/Source/Runtime/TypedElementFramework/Public/Elements/Interfaces/TypedElementDataStorageInterface.h
+++ b/Engine/Source/Runtime/TypedElementFramework/Public/Elements/Interfaces/TypedElementDataStorageInterface.h
@@ -64,24 +64,25 @@ class UTypedElementDataStorageInterface : public UInterface
 	GENERATED_BODY()
 };
 
+/**
+ * Convenience structure that can be used to pass a list of columns to functions that don't
+ * have an dedicate templated version that takes a column list directly, for instance when
+ * multiple column lists are used. Note that the returned array view is only available while
+ * this object is constructed, so care must be taken with functions that return a const array view.
+ */
+template<typename... Columns>
+struct TTypedElementColumnTypeList
+{
+	const UScriptStruct* ColumnTypes[sizeof...(Columns)] = { Columns::StaticStruct()... };
+	
+	operator TConstArrayView<const UScriptStruct*>() { return ColumnTypes; }
+};
+
 class TYPEDELEMENTFRAMEWORK_API ITypedElementDataStorageInterface
 {
 	GENERATED_BODY()
 
 public:
-	/** 
-	 * Convenience structure that can be used to pass a lis of columns to functions that don't
-	 * have an dedicate templated version that takes a column list directly, for instance when
-	 * multiple column lists are used.
-	 */
-	template<typename... Columns>
-	struct TColumnTypeList
-	{
-		const UScriptStruct* ColumnTypes[sizeof...(Columns)] = { Columns::StaticStruct()... };
-
-		operator TConstArrayView<const UScriptStruct*>() { return ColumnTypes; }
-	};
-
 	/**
 	 * @section Table management
 	 * 
@@ -180,7 +181,7 @@ public:
 	virtual void* GetColumnData(TypedElementRowHandle Row, const UScriptStruct* ColumnType) = 0;
 	virtual ColumnDataResult GetColumnData(TypedElementRowHandle Row, FTopLevelAssetPath ColumnName) = 0;
 
-	/** Determines if the provided row contains the collection of columsn and tags. */
+	/** Determines if the provided row contains the collection of columns and tags. */
 	virtual bool HasColumns(TypedElementRowHandle Row, TConstArrayView<const UScriptStruct*> ColumnTypes) const = 0;
 	virtual bool HasColumns(TypedElementRowHandle Row, TConstArrayView<TWeakObjectPtr<const UScriptStruct>> ColumnTypes) const = 0;
 	
@@ -189,7 +190,7 @@ public:
 	 * @section Query
 	 * @description
 	 * Queries can be constructed using the Query Builder. Note that the Query Builder allows for the creation of queries that
-	 * are more complex than the backend may support. The backend is allowed to simplify the query, in which case the query
+	 * are more complex than the back-end may support. The back-end is allowed to simplify the query, in which case the query
 	 * can be used directly in the processor to do additional filtering. This will however impact performance and it's 
 	 * therefore recommended to try to simplify the query first before relying on extended query filtering in a processor.
 	 */
@@ -215,7 +216,7 @@ public:
 		 */
 		SyncExternalToDataStorage,
 		/**
-		 * The group for quries that need to sync data from the Data Storage to external sources such as subsystems
+		 * The group for queries that need to sync data from the Data Storage to external sources such as subsystems
 		 * or the world into. These typically run late in a phase.
 		 */
 		SyncDataStorageToExternal,
@@ -265,9 +266,9 @@ public:
 		None = 0,
 		/** If set the dependency is accessed as read-only. If not set the dependency requires Read/Write access. */
 		ReadOnly = 1 << 0,
-		/** If set the dependecy can only be used from the game thread, otherwise it can be accessed from any thread. */
+		/** If set the dependency can only be used from the game thread, otherwise it can be accessed from any thread. */
 		GameThreadBound = 1 << 1,
-		/** If set the dependency will be refetched every iteration, otherwise only if not fetched before. */
+		/** If set the dependency will be re-fetched every iteration, otherwise only if not fetched before. */
 		AlwaysRefresh = 1 << 2
 	};
 
@@ -283,13 +284,13 @@ public:
 		/** Return the address of a mutable column matching the requested type or a nullptr if not found. */
 		virtual void* GetMutableColumn(const UScriptStruct* ColumnType) = 0;
 		/**
-		 * Get a list of columns or nullptrs if the columntype wasn't found. Mutable addresses are returned and it's up to
+		 * Get a list of columns or nullptrs if the column type wasn't found. Mutable addresses are returned and it's up to
 		 * the caller to not change immutable addresses.
 		 */
 		virtual void GetColumns(TArrayView<char*> RetrievedAddresses, TConstArrayView<TWeakObjectPtr<const UScriptStruct>> ColumnTypes,
 			TConstArrayView<ITypedElementDataStorageInterface::EQueryAccessType> AccessTypes) = 0;
 		/**
-		 * Get a list of columns or nullptrs if the columntype wasn't found. Mutable addresses are returned and it's up to
+		 * Get a list of columns or nullptrs if the column type wasn't found. Mutable addresses are returned and it's up to
 		 * the caller to not change immutable addresses. This version doesn't verify that the enough space is provided and
 		 * it's up to the caller to guarantee the target addresses have enough space.
 		 */
@@ -490,7 +491,7 @@ public:
 			/** Only portions of the query were executed. This is caused by a problem that was encountered partway through processing. */
 			Partially,
 			/** 
-			 * The backend doesn't support the particular query. This may be a limitation in how/where the query is run or because
+			 * The back-end doesn't support the particular query. This may be a limitation in how/where the query is run or because
 			 * the query contains actions and/or operations that are not supported.
 			 */
 			Unsupported,
diff --git a/Engine/Source/Runtime/TypedElementFramework/Public/Elements/Interfaces/TypedElementDataStorageUiInterface.h b/Engine/Source/Runtime/TypedElementFramework/Public/Elements/Interfaces/TypedElementDataStorageUiInterface.h
--- a/Engine/Source/Runtime/TypedElementFramework/Public/Elements/Interfaces/TypedElementDataStorageUiInterface.h
+++ b/Engine/Source/Runtime/TypedElementFramework/Public/Elements/Interfaces/TypedElementDataStorageUiInterface.h
@@ -29,7 +29,7 @@ struct TYPEDELEMENTFRAMEWORK_API FTypedElementWidgetConstructor
 	GENERATED_BODY()
 
 public:
-	explicit FTypedElementWidgetConstructor(UScriptStruct* InTypeInfo);
+	explicit FTypedElementWidgetConstructor(const UScriptStruct* InTypeInfo);
 	explicit FTypedElementWidgetConstructor(EForceInit) {} //< For compatibility and shouldn't be directly used.
 
 	virtual ~FTypedElementWidgetConstructor() = default;
@@ -37,9 +37,20 @@ public:
 	/** Retrieves the type information for the constructor type. */
 	virtual const UScriptStruct* GetTypeInfo() const;
 
+	/** Returns a list of additional columns the widget requires. */
+	virtual TConstArrayView<const UScriptStruct*> GetAdditionalColumnsList() const;
+
+	/** 
+	 * Whether or not an instance of this constructor can be reused. Setting this to true means that the Data Storage will
+	 * avoid recreating new instances of the constructor and instead keeps reusing the created version. The default is
+	 * 'false'. Override this function and set it to true if there's no internal state that could interfere with creating
+	 * multiple widgets from the same constructor, e.g. data stored from passed in arguments.
+	 */
+	virtual bool CanBeReused() const;
+
 	/**
 	 * Constructs the widget according to the provided information. Information is collected by calling
-	 * the below functions ApplyArguments, CreateWidget and AddColumns. It's recommeded to overload those
+	 * the below functions ApplyArguments, CreateWidget and AddColumns. It's recommended to overload those
 	 * functions to build widgets according to a standard recipe and to reduce the amount of code needed.
 	 * If a complexer situation is called for this function can also be directly overwritten.
 	 */
@@ -55,16 +66,22 @@ protected:
 	 * function if there are non-properties that can be configured through the arguments or if the configuration
 	 * contains arguments that are too complex to initialize (fully) through the type system.
 	 */
-	virtual void ApplyArguments(TConstArrayView<TypedElement::ColumnUtils::Argument> Arguments);
+	virtual bool ApplyArguments(TConstArrayView<TypedElement::ColumnUtils::Argument> Arguments);
 	/** Create a new instance of the target widget. This is a required function. */
 	virtual TSharedPtr<SWidget> CreateWidget() PURE_VIRTUAL(FTypedElementWidgetConstructor::CreateWidget, return nullptr; );
-	/**
-	 * Add any columns to the provided row in the data storage to support the widget. By default a column will be added
-	 * that keeps track of the lifecycle of the widget.
+	/** Set any values in columns if needed. The columns provided through GetAdditionalColumnsList() will have already been created. */
+	virtual bool SetColumns(ITypedElementDataStorageInterface* DataStorage, TypedElementRowHandle Row);
+	/** 
+	 * Last opportunity to configure anything in the widget or the row. This step can be needed to initialize widgets with data stored
+	 * in columns.
 	 */
-	virtual void AddColumns(ITypedElementDataStorageInterface* DataStorage, TypedElementRowHandle Row, const TSharedPtr<SWidget>& Widget);
+	virtual bool FinalizeWidget(
+		ITypedElementDataStorageInterface* DataStorage,
+		ITypedElementDataStorageUiInterface* DataStorageUi,
+		TypedElementRowHandle Row,
+		const TSharedPtr<SWidget>& Widget);
 
-	UScriptStruct* TypeInfo{ nullptr };
+	const UScriptStruct* TypeInfo{ nullptr };
 };
 
 template<>
@@ -93,17 +110,30 @@ public:
 		/** General purpose name which allows multiple factory registrations. */
 		Generic,
 		/**
-		 * Only one factory can be registrated with this purpose. If multiple factories are registered only the last factory
+		 * Only one factory can be registered with this purpose. If multiple factories are registered only the last factory
 		 * will be stored.
 		 */
 		UniqueByName,
 		/**
-		 * Only one factory can be registrated with this purpose for a specific combination of columns. If multiple factories
+		 * Only one factory can be registered with this purpose for a specific combination of columns. If multiple factories
 		 * are registered only the last factory will be stored.
 		 */
 		UniqueByNameAndColumn
 	};
 
+	enum class EMatchApproach : uint8
+	{
+		/** 
+		 * Looks for the longest chain of columns matching widget factories. The matching columns are removed and the process
+		 * is repeated until there are no more columns or no matches are found.
+		 */
+		LongestMatch,
+		/** A single widget factory is reduced which matches the requested columns exactly. */
+		ExactMatch,
+		/** Each column is matched to widget factory. Only widget factories that use a single column are used. */
+		SingleMatch
+	};
+
 	using WidgetCreatedCallback = TFunctionRef<void(const TSharedRef<SWidget>& NewWidget, TypedElementRowHandle Row)>;
 	using WidgetConstructorCallback = TFunctionRef<bool(TUniquePtr<FTypedElementWidgetConstructor>, TConstArrayView<TWeakObjectPtr<const UScriptStruct>>)>;
 	using WidgetPurposeCallback = TFunctionRef<void(FName, EPurposeType, const FText&)>;
@@ -119,7 +149,7 @@ public:
 	 * This version registers a generic type. Construction using these are typically cheaper as they can avoid
 	 * copying the Constructor and take up less memory. The downside is that they can't store additional configuration
 	 * options. If the purpose has not been registered the factory will not be recorded and a warning will be printed.
-	 * If registration is successfull true will be returned otherwise false.
+	 * If registration is successful true will be returned otherwise false.
 	 */
 	virtual bool RegisterWidgetFactory(FName Purpose, const UScriptStruct* Constructor) = 0;
 	/**
@@ -128,27 +158,28 @@ public:
 	 * copying the Constructor and take up less memory. The downside is that they can't store additional configuration
 	 * options. If the purpose has not been registered the factory will not be recorded and a warning will be printed.
 	 * The provided columns will be used when matching the factory during widget construction.
-	 * If registration is successfull true will be returned otherwise false.
+	 * If registration is successful true will be returned otherwise false.
 	 */
 	virtual bool RegisterWidgetFactory(FName Purpose, const UScriptStruct* Constructor,
 		TArray<TWeakObjectPtr<const UScriptStruct>> Columns) = 0;
 	/**
 	 * Registers a widget factory that will be called when the purpose it's registered under is requested.
-	 * This version uses a pre-created instance of the Constructor. The benefit of this is that it store
+	 * This version uses a previously created instance of the Constructor. The benefit of this is that it store
 	 * configuration options. The downside is that it takes up more memory and requires copying when it's
 	 * used. If the purpose has not been registered the factory will not be recorded and a warning will be printed.
-	 * If registration is successfull true will be returned otherwise false.
+	 * If registration is successful true will be returned otherwise false.
 	 */
 	virtual bool RegisterWidgetFactory(FName Purpose, TUniquePtr<FTypedElementWidgetConstructor>&& Constructor) = 0;
 	/**
 	 * Registers a widget factory that will be called when the purpose it's registered under is requested.
-	 * This version uses a pre-created instance of the Constructor. The benefit of this is that it store
+	 * This version uses a previously created instance of the Constructor. The benefit of this is that it store
 	 * configuration options. The downside is that it takes up more memory and requires copying when it's
 	 * used. If the purpose has not been registered the factory will not be recorded and a warning will be printed.
 	 * The provided columns will be used when matching the factory during widget construction.
-	 * If registration is successfull true will be returned otherwise false.
+	 * If registration is successful true will be returned otherwise false.
 	 */
-	virtual bool RegisterWidgetFactory(FName Purpose, TUniquePtr<FTypedElementWidgetConstructor>&& Constructor, TArray<TWeakObjectPtr<const UScriptStruct>> Columns) = 0;
+	virtual bool RegisterWidgetFactory(FName Purpose, TUniquePtr<FTypedElementWidgetConstructor>&& Constructor, 
+		TArray<TWeakObjectPtr<const UScriptStruct>> Columns) = 0;
 	
 	/** Creates widget constructors for the requested purpose. */
 	virtual void CreateWidgetConstructors(FName Purpose, 
@@ -157,7 +188,7 @@ public:
 	 * Finds matching widget constructors for provided columns, preferring longer matches over shorter matches.
 	 * The provided list of columns will be updated to contain all columns that couldn't be matched.
 	 */
-	virtual void CreateWidgetConstructors(FName Purpose, TArray<TWeakObjectPtr<const UScriptStruct>>& Columns,
+	virtual void CreateWidgetConstructors(FName Purpose, EMatchApproach MatchApproach, TArray<TWeakObjectPtr<const UScriptStruct>>& Columns,
 		TConstArrayView<TypedElement::ColumnUtils::Argument> Arguments, const WidgetConstructorCallback& Callback) = 0;
 
 	/**
