Fixed an initialization bug in TEDS for cooked builds.

The order-of-initialization for cooked editors appears to be different from standalone builds. As a result the static initialization trick used to impersonate Mass tags and fragments doesn't work. However since TEDS no only uses dynamic processors there's no longer a need to have the impersonation initialized before the CDOs are created (which would setup queries which in turn need tags and fragments).

diff --git a/Engine/Plugins/Experimental/TypedElementsDataStorage/Source/TypedElementsDataStorage/Private/TypedElementsDataStorage.cpp b/Engine/Plugins/Experimental/TypedElementsDataStorage/Source/TypedElementsDataStorage/Private/TypedElementsDataStorage.cpp
--- a/Engine/Plugins/Experimental/TypedElementsDataStorage/Source/TypedElementsDataStorage/Private/TypedElementsDataStorage.cpp
+++ b/Engine/Plugins/Experimental/TypedElementsDataStorage/Source/TypedElementsDataStorage/Private/TypedElementsDataStorage.cpp
@@ -15,41 +15,41 @@
 #define LOCTEXT_NAMESPACE "FTypedElementsDataStorageModule"
 
 // MASS uses CDO in a few places, making it a difficult to consistently register Type Element's Columns and Tags
-// as they may have not been set up to impersonate MASS' Fragments and Tags yet. To get around this, do the 
-// registration during static initialization, which is before CDOs.
+// as they may have not been set up to impersonate MASS' Fragments and Tags yet. There are currently no longer
+// any cases where TEDS relies on this but may happen again the future. For the standalone version a static
+// can be used to initialize the impersonation before the CDO get a chance to run, but for a cooked editor this will
+// not work.
 /**
  * Typed Elements provides base classes for columns and tags. These directly map to fragments and tags in MASS.
  * To avoid deep and tight coupling between both systems, columns and tags don't directly inhered from MASS, but
  * are otherwise fully compatible. To allow MASS to do its type safety checks, this class updates the type
  * information so Typed Elements columns and tags present as MASS fragments and tags from MASS's perspective.
  */
-struct FTagAndFragmentImpersonators
+void ImpersonateMassTagsAndFragments()
 {
-	FTagAndFragmentImpersonators()
-	{
-		// Have FTypedElementDataStorageColumn impersonate a FMassFragment, which is the actual data storage when using MASS as a backend.
-		static_assert(sizeof(FTypedElementDataStorageColumn) == sizeof(FMassFragment),
-			"In order for FTypedElementDataStorageColumn to impersonate FMassFragment they need to be identical.");
-		static_assert(!TIsPolymorphic<FMassFragment>::Value,
-			"In order to be able to impersonate FMassFragment it can't have any virtual functions.");
-		static_assert(!TIsPolymorphic<FTypedElementDataStorageColumn>::Value,
-			"In order to be able to use FTypedElementDataStorageColumn to impersonate FMassFragment it can't have any virtual functions.");
-		FTypedElementDataStorageColumn::StaticStruct()->SetSuperStruct(FMassFragment::StaticStruct());
-
-		// Have FTypedElementDataStorageTag impersonate a FMassTag, which is the tag type when using MASS as a backend.
-		static_assert(sizeof(FTypedElementDataStorageTag) == sizeof(FMassTag),
-			"In order for FTypedElementDataStorageTag to impersonate FMassTag they need to be identical.");
-		static_assert(!TIsPolymorphic<FMassTag>::Value,
-			"In order to be able to impersonate FMassTag it can't have any virtual functions.");
-		static_assert(!TIsPolymorphic<FTypedElementDataStorageTag>::Value,
-			"In order to be able to use FTypedElementDataStorageTag to impersonate FMassTag it can't have any virtual functions.");
-		FTypedElementDataStorageTag::StaticStruct()->SetSuperStruct(FMassTag::StaticStruct());
-	}
-};
-static FTagAndFragmentImpersonators TagAndFragmentImpersonators;
+	// Have FTypedElementDataStorageColumn impersonate a FMassFragment, which is the actual data storage when using MASS as a backend.
+	static_assert(sizeof(FTypedElementDataStorageColumn) == sizeof(FMassFragment),
+		"In order for FTypedElementDataStorageColumn to impersonate FMassFragment they need to be identical.");
+	static_assert(!TIsPolymorphic<FMassFragment>::Value,
+		"In order to be able to impersonate FMassFragment it can't have any virtual functions.");
+	static_assert(!TIsPolymorphic<FTypedElementDataStorageColumn>::Value,
+		"In order to be able to use FTypedElementDataStorageColumn to impersonate FMassFragment it can't have any virtual functions.");
+	FTypedElementDataStorageColumn::StaticStruct()->SetSuperStruct(FMassFragment::StaticStruct());
+
+	// Have FTypedElementDataStorageTag impersonate a FMassTag, which is the tag type when using MASS as a backend.
+	static_assert(sizeof(FTypedElementDataStorageTag) == sizeof(FMassTag),
+		"In order for FTypedElementDataStorageTag to impersonate FMassTag they need to be identical.");
+	static_assert(!TIsPolymorphic<FMassTag>::Value,
+		"In order to be able to impersonate FMassTag it can't have any virtual functions.");
+	static_assert(!TIsPolymorphic<FTypedElementDataStorageTag>::Value,
+		"In order to be able to use FTypedElementDataStorageTag to impersonate FMassTag it can't have any virtual functions.");
+	FTypedElementDataStorageTag::StaticStruct()->SetSuperStruct(FMassTag::StaticStruct());
+}
 
 void FTypedElementsDataStorageModule::StartupModule()
 {
+	ImpersonateMassTagsAndFragments();
+
 	FCoreDelegates::OnAllModuleLoadingPhasesComplete.AddLambda(
 		[this]()
 		{
