Added Scene Outliner integration for TEDS.

The integration is currently limited to adding the columns from a select query and only cells in a row with actors are filled in. The console command "TEDS.UI.BindColumnsToSceneOutliner" was added to create a column-only select query and bind it with the most recently used Scene Outliner. A default widget is provided that will take any column and use the type info to export it as plain text.

diff --git a/Engine/Plugins/Experimental/TypedElementsDataStorage/Source/TypedElementsDataStorageUI/Private/TypedElementExportedTextWidget.cpp b/Engine/Plugins/Experimental/TypedElementsDataStorage/Source/TypedElementsDataStorageUI/Private/TypedElementExportedTextWidget.cpp
new file mode 100644
--- /dev/null
+++ b/Engine/Plugins/Experimental/TypedElementsDataStorage/Source/TypedElementsDataStorageUI/Private/TypedElementExportedTextWidget.cpp
@@ -0,0 +1,119 @@
+#include "TypedElementExportedTextWidget.h"
+
+#include "Elements/Columns/TypedElementMiscColumns.h"
+#include "Elements/Columns/TypedElementSlateWidgetColumns.h"
+#include "Elements/Columns/TypedElementTypeInfoColumns.h"
+#include "Elements/Framework/TypedElementQueryBuilder.h"
+#include "TypedElementSubsystems.h"
+#include "Widgets/Text/STextBlock.h"
+
+#define LOCTEXT_NAMESPACE "TypedElementUI_ExportedTextWidget"
+
+//
+// UTypedElementExportedTextWidgetFactory
+//
+
+static void UpdateExportedTextWidget(ITypedElementDataStorageInterface& DataStorage, FTypedElementSlateWidgetReferenceColumn& Widget,
+	const FTypedElementScriptStructTypeInfoColumn& TypeInfo, const FTypedElementRowReferenceColumn& ReferencedRow)
+{
+	TSharedPtr<SWidget> WidgetPointer = Widget.Widget.Pin();
+	checkf(WidgetPointer, TEXT("Referenced widget is not valid. A constructed widget may not have been cleaned up. This can "
+		"also happen if this processor is running in the same phase as the processors responsible for cleaning up old "
+		"references."));
+	checkf(WidgetPointer->GetType() == STextBlock::StaticWidgetClass().GetWidgetType(),
+		TEXT("Stored widget with FTypedElementExportedTextWidgetTag doesn't match type %s, but was a %s."),
+		*(STextBlock::StaticWidgetClass().GetWidgetType().ToString()),
+		*(WidgetPointer->GetTypeAsString()));
+
+	if (void* Data = DataStorage.GetColumnData(ReferencedRow.Row, TypeInfo.TypeInfo.Get()))
+	{
+		FString Label;
+		TypeInfo.TypeInfo->ExportText(Label, Data, nullptr, nullptr, PPF_None, nullptr);
+		STextBlock* TextWidget = static_cast<STextBlock*>(WidgetPointer.Get());
+		FText Text = FText::FromString(MoveTemp(Label));
+		TextWidget->SetToolTipText(Text);
+		TextWidget->SetText(MoveTemp(Text));
+	}
+}
+
+void UTypedElementExportedTextWidgetFactory::RegisterQueries(ITypedElementDataStorageInterface& DataStorage) const
+{
+	using namespace TypedElementQueryBuilder;
+	using DSI = ITypedElementDataStorageInterface;
+
+	DataStorage.RegisterQuery(
+		Select(TEXT("Sync exported text widgets"),
+		FProcessor(DSI::EQueryTickPhase::FrameEnd, DataStorage.GetQueryTickGroupName(DSI::EQueryTickGroups::SyncWidgets))
+			.ForceToGameThread(true),
+			[](FCachedQueryContext<UTypedElementDataStorageSubsystem>& Context, FTypedElementSlateWidgetReferenceColumn& Widget,
+				const FTypedElementScriptStructTypeInfoColumn& TypeInfo, const FTypedElementRowReferenceColumn& ReferencedRow)
+			{
+				UTypedElementDataStorageSubsystem& Subsystem = Context.GetCachedMutableDependency<UTypedElementDataStorageSubsystem>();
+				DSI* DataStorage = Subsystem.Get();
+				checkf(DataStorage, TEXT("FTypedElementsDataStorageUiModule tried to process widgets before the "
+					"Typed Elements Data Storage interface is available."));
+				
+				if (DataStorage->HasColumns<FTypedElementSyncFromWorldTag>(ReferencedRow.Row) ||
+					DataStorage->HasColumns<FTypedElementSyncBackToWorldTag>(ReferencedRow.Row))
+				{
+					UpdateExportedTextWidget(*DataStorage, Widget, TypeInfo, ReferencedRow);
+				}
+			}
+		)
+	.Where()
+		.All<FTypedElementExportedTextWidgetTag>()
+	.Compile());
+
+}
+
+void UTypedElementExportedTextWidgetFactory::RegisterWidgetConstructors(ITypedElementDataStorageInterface& DataStorage,
+	ITypedElementDataStorageUiInterface& DataStorageUi) const
+{
+	DataStorageUi.RegisterWidgetFactory(FName(TEXT("General.Cell.Default")), FTypedElementExportedTextWidgetConstructor::StaticStruct());
+}
+
+
+
+//
+// FTypedElementExportedTextWidgetConstructor
+//
+
+FTypedElementExportedTextWidgetConstructor::FTypedElementExportedTextWidgetConstructor()
+	: Super(FTypedElementExportedTextWidgetConstructor::StaticStruct())
+{
+}
+
+TConstArrayView<const UScriptStruct*> FTypedElementExportedTextWidgetConstructor::GetAdditionalColumnsList() const
+{
+	static TTypedElementColumnTypeList<
+		FTypedElementRowReferenceColumn,
+		FTypedElementScriptStructTypeInfoColumn,
+		FTypedElementExportedTextWidgetTag> Columns;
+	return Columns;
+}
+
+bool FTypedElementExportedTextWidgetConstructor::CanBeReused() const
+{
+	return true;
+}
+
+TSharedPtr<SWidget> FTypedElementExportedTextWidgetConstructor::CreateWidget()
+{
+	return SNew(STextBlock);
+}
+
+bool FTypedElementExportedTextWidgetConstructor::FinalizeWidget(
+	ITypedElementDataStorageInterface* DataStorage,
+	ITypedElementDataStorageUiInterface* DataStorageUi,
+	TypedElementRowHandle Row,
+	const TSharedPtr<SWidget>& Widget)
+{
+	UpdateExportedTextWidget(
+		*DataStorage,
+		*DataStorage->GetColumn<FTypedElementSlateWidgetReferenceColumn>(Row),
+		*DataStorage->GetColumn<FTypedElementScriptStructTypeInfoColumn>(Row),
+		*DataStorage->GetColumn<FTypedElementRowReferenceColumn>(Row));
+	return true;
+}
+
+#undef LOCTEXT_NAMESPACE
diff --git a/Engine/Plugins/Experimental/TypedElementsDataStorage/Source/TypedElementsDataStorageUI/Private/TypedElementGeneralRegistration.cpp b/Engine/Plugins/Experimental/TypedElementsDataStorage/Source/TypedElementsDataStorageUI/Private/TypedElementGeneralRegistration.cpp
new file mode 100644
--- /dev/null
+++ b/Engine/Plugins/Experimental/TypedElementsDataStorage/Source/TypedElementsDataStorageUI/Private/TypedElementGeneralRegistration.cpp
@@ -0,0 +1,23 @@
+#include "TypedElementGeneralRegistration.h"
+
+#include "Elements/Interfaces/TypedElementDataStorageUiInterface.h"
+
+#define LOCTEXT_NAMESPACE "TypedElementsUI_GeneralRegistration"
+
+const FName UTypedElementGeneralRegistrationFactory::CellPurpose(TEXT("General.Cell"));
+const FName UTypedElementGeneralRegistrationFactory::CellDefaultPurpose(TEXT("General.Cell.Default"));
+const FName UTypedElementGeneralRegistrationFactory::HeaderDefaultPurpose(TEXT("General.Header.Default"));
+
+void UTypedElementGeneralRegistrationFactory::RegisterWidgetPurposes(ITypedElementDataStorageUiInterface& DataStorageUi) const
+{
+	DataStorageUi.RegisterWidgetPurpose(CellPurpose, ITypedElementDataStorageUiInterface::EPurposeType::UniqueByNameAndColumn,
+		LOCTEXT("TypedElementsUI_GeneralRegistration", "General purpose widgets that can be used as cells for specific columns or column combinations."));
+	
+	DataStorageUi.RegisterWidgetPurpose(HeaderDefaultPurpose, ITypedElementDataStorageUiInterface::EPurposeType::UniqueByName,
+		LOCTEXT("TypedElementsUI_GeneralRegistration", "The default widget to use in headers if no other specialization is provided."));
+	
+	DataStorageUi.RegisterWidgetPurpose(CellDefaultPurpose, ITypedElementDataStorageUiInterface::EPurposeType::UniqueByName,
+		LOCTEXT("TypedElementsUI_GeneralRegistration", "The default widget to use in cells if no other specialization is provided."));
+}
+
+#undef LOCTEXT_NAMESPACE
diff --git a/Engine/Plugins/Experimental/TypedElementsDataStorage/Source/TypedElementsDataStorageUI/Private/TypedElementSceneOutliner.cpp b/Engine/Plugins/Experimental/TypedElementsDataStorage/Source/TypedElementsDataStorageUI/Private/TypedElementSceneOutliner.cpp
new file mode 100644
--- /dev/null
+++ b/Engine/Plugins/Experimental/TypedElementsDataStorage/Source/TypedElementsDataStorageUI/Private/TypedElementSceneOutliner.cpp
@@ -0,0 +1,537 @@
+// Copyright Epic Games, Inc. All Rights Reserved.
+
+#include "TypedElementSceneOutliner.h"
+
+#include "ActorTreeItem.h"
+#include "Elements/Columns/TypedElementLabelColumns.h"
+#include "Elements/Columns/TypedElementMiscColumns.h"
+#include "Elements/Columns/TypedElementPackageColumns.h"
+#include "Elements/Columns/TypedElementSlateWidgetColumns.h"
+#include "Elements/Columns/TypedElementTransformColumns.h"
+#include "Elements/Columns/TypedElementTypeInfoColumns.h"
+#include "Elements/Framework/TypedElementRegistry.h"
+#include "Elements/Framework/TypedElementQueryBuilder.h"
+#include "Elements/Interfaces/TypedElementDataStorageCompatibilityInterface.h"
+#include "ILevelEditor.h"
+#include "ISceneOutliner.h"
+#include "ISceneOutlinerColumn.h"
+#include "ISceneOutlinerTreeItem.h"
+#include "LevelEditor.h"
+#include "SSceneOutliner.h"
+#include "SceneOutlinerFwd.h"
+#include "SceneOutlinerPublicTypes.h"
+#include "TypedElementSubsystems.h"
+
+#define LOCTEXT_NAMESPACE "TypedElementsUI_SceneOutliner"
+
+FAutoConsoleCommand BindColumnsToSceneOutlinerConsoleCommand(
+	TEXT("TEDS.UI.BindColumnsToSceneOutliner"),
+	TEXT("Bind one or more columns to the most recently used Scene Outliner. Several prebuild configurations are offered as well.")
+	TEXT("An example input to show a label column is 'TEDS.UI.BindColumnsToSceneOutliner /Script/TypedElementFramework.TypedElementLabelColumn'."),
+	FConsoleCommandWithArgsDelegate::CreateLambda([](const TArray<FString>& Args)
+		{
+			using namespace TypedElementQueryBuilder;
+			using DSI = ITypedElementDataStorageInterface;
+
+			UTypedElementRegistry* Registry = UTypedElementRegistry::GetInstance();
+			if (ITypedElementDataStorageInterface* DataStorage = Registry->GetMutableDataStorage())
+			{
+				static TypedElementQueryHandle Queries[] =
+				{
+					DataStorage->RegisterQuery(Select().ReadWrite<FTypedElementLabelColumn>().Compile()),
+					DataStorage->RegisterQuery(Select().ReadOnly<FTypedElementLocalTransformColumn>().Compile()),
+					DataStorage->RegisterQuery(Select().ReadOnly<FTypedElementPackagePathColumn>().Compile()),
+					DataStorage->RegisterQuery(
+						Select()
+							.ReadWrite<FTypedElementLabelColumn>()
+							.ReadOnly<FTypedElementLocalTransformColumn>()
+						.Compile()),
+					DataStorage->RegisterQuery(
+						Select()
+							.ReadWrite<FTypedElementLabelColumn>()
+							.ReadOnly<FTypedElementLabelHashColumn>()
+						.Compile())
+				};
+
+				FTypedElementSceneOutlinerQueryBinder& Binder = FTypedElementSceneOutlinerQueryBinder::GetInstance();
+				const TWeakPtr<ILevelEditor> LevelEditor = FModuleManager::GetModuleChecked<FLevelEditorModule>(TEXT("LevelEditor")).GetLevelEditorInstance();
+				const TSharedPtr<ISceneOutliner> SceneOutliner = LevelEditor.IsValid() ? LevelEditor.Pin()->GetMostRecentlyUsedSceneOutliner() : nullptr;
+				if (SceneOutliner.IsValid())
+				{
+					if (!Args.IsEmpty())
+					{
+						if (Args[0].IsNumeric())
+						{
+							int32 QueryIndex = FCString::Atoi(*Args[0]);
+							if (QueryIndex < sizeof(Queries) / sizeof(TypedElementQueryHandle))
+							{
+								Binder.AssignQuery(Queries[QueryIndex], SceneOutliner);
+								return;
+							}
+						}
+						else
+						{
+							uint32 AdditionCount = 0;
+							Select Query;
+							for (const FString& Arg : Args)
+							{
+								FTopLevelAssetPath Path;
+								// TrySetPath has an ensure that checks if the path starts with an '/' and otherwise throws
+								// an assert.
+								if (!Arg.IsEmpty() && Arg[0] == '/' && Path.TrySetPath(Arg))
+								{
+									const UScriptStruct* ColumnType = TypeOptional(Path);
+									if (ColumnType)
+									{
+										Query.ReadOnly(ColumnType);
+										++AdditionCount;
+									}
+								}
+							}
+							if (AdditionCount > 0)
+							{
+								static TypedElementQueryHandle CustomQuery = TypedElementInvalidQueryHandle;
+								if (CustomQuery != TypedElementInvalidQueryHandle)
+								{
+									DataStorage->UnregisterQuery(CustomQuery);
+								}
+								CustomQuery = DataStorage->RegisterQuery(Query.Compile());
+								Binder.AssignQuery(CustomQuery, SceneOutliner);
+								return;
+							}
+						}
+					}
+					Binder.AssignQuery(TypedElementInvalidQueryHandle, SceneOutliner);
+				}
+			}
+		}));
+
+class FOutlinerColumn : public ISceneOutlinerColumn
+{
+public:
+	FOutlinerColumn(
+		ITypedElementDataStorageInterface& InStorage, 
+		ITypedElementDataStorageUiInterface& InStorageUi, 
+		ITypedElementDataStorageCompatibilityInterface& InStorageCompatibility,
+		FName InNameId, 
+		TArray<TWeakObjectPtr<const UScriptStruct>> InColumnTypes,
+		TSharedPtr<FTypedElementWidgetConstructor> InHeaderWidgetConstructor,
+		TSharedPtr<FTypedElementWidgetConstructor> InCellWidgetConstructor)
+		: ColumnTypes(MoveTemp(InColumnTypes))
+		, HeaderWidgetConstructor(MoveTemp(InHeaderWidgetConstructor))
+		, CellWidgetConstructor(MoveTemp(InCellWidgetConstructor))
+		, Storage(InStorage)
+		, StorageUi(InStorageUi)
+		, StorageCompatibility(InStorageCompatibility)
+		, NameId(InNameId)
+	{
+		ColumnTypes.Shrink();
+	};
+	
+	~FOutlinerColumn() override = default;
+
+	FName GetColumnID() override
+	{
+		return NameId;
+	}
+
+	SHeaderRow::FColumn::FArguments ConstructHeaderRowColumn() override
+	{
+		FString TooltipText = TEXT("Data Storage columns:");
+		for (const TWeakObjectPtr<const UScriptStruct>& ColumnType : ColumnTypes)
+		{
+			if (ColumnType.IsValid())
+			{
+				TooltipText += TEXT("\n    ");
+				ColumnType->AppendName(TooltipText);
+			}
+		}
+
+		TSharedPtr<SWidget> Widget;
+		if (HeaderWidgetConstructor)
+		{
+			TypedElementRowHandle UiRowHandle = Storage.AddRow(FName(TEXT("Editor_WidgetTable")));
+			Widget = StorageUi.ConstructWidget(UiRowHandle, *HeaderWidgetConstructor, 
+				{
+					{ FName(TEXT("Name")), NameId.ToString() }
+				});
+		}
+		if (!Widget.IsValid())
+		{
+			Widget = SNew(STextBlock)
+				.Text(FText::FromString(NameId.ToString()));
+		}
+		
+		return SHeaderRow::Column(NameId)
+			.FillWidth(2)
+			.HeaderComboVisibility(EHeaderComboVisibility::OnHover)
+			.DefaultTooltip(FText::FromString(MoveTemp(TooltipText)))
+			.HeaderContent()
+			[
+				SNew(SHorizontalBox)
+				+ SHorizontalBox::Slot()
+				.VAlign(VAlign_Center)
+				[
+					Widget.ToSharedRef()
+				]
+			];
+	}
+	
+	const TSharedRef<SWidget> ConstructRowWidget(FSceneOutlinerTreeItemRef TreeItem, const STableRow<FSceneOutlinerTreeItemPtr>& Row) override
+	{
+		if (const FActorTreeItem* ActorItem = TreeItem->CastTo<FActorTreeItem>())
+		{
+			if (AActor* Actor = ActorItem->Actor.Get())
+			{
+				TypedElementRowHandle RowHandle = StorageCompatibility.FindRowWithCompatibleObject(Actor);
+				if (RowHandle != TypedElementInvalidRowHandle && Storage.HasColumns(RowHandle, ColumnTypes))
+				{
+					TypedElementRowHandle UiRowHandle = Storage.AddRow(FTypedElementSceneOutlinerQueryBinder::CellWidgetTableName);
+					Storage.AddColumns(UiRowHandle, CellWidgetConstructor->GetAdditionalColumnsList());
+					if (ColumnTypes.Num() == 1)
+					{
+						if (FTypedElementScriptStructTypeInfoColumn* TypeInfo = Storage.GetColumn<FTypedElementScriptStructTypeInfoColumn>(UiRowHandle))
+						{
+							TypeInfo->TypeInfo = *ColumnTypes.begin();
+						}
+					}
+					if (FTypedElementRowReferenceColumn* RowReference = Storage.GetColumn<FTypedElementRowReferenceColumn>(UiRowHandle))
+					{
+						RowReference->Row = RowHandle;
+					}
+
+					if (TSharedPtr<SWidget> Widget = StorageUi.ConstructWidget(UiRowHandle, *CellWidgetConstructor, {}))
+					{
+						return Widget.ToSharedRef();
+					}
+					else
+					{
+						Storage.RemoveRow(RowHandle);
+					}
+				}
+			}
+		}
+		return SNullWidget::NullWidget;
+	}
+
+	TArray<TWeakObjectPtr<const UScriptStruct>> ColumnTypes;
+	TSharedPtr<FTypedElementWidgetConstructor> HeaderWidgetConstructor;
+	TSharedPtr<FTypedElementWidgetConstructor> CellWidgetConstructor;
+	ITypedElementDataStorageInterface& Storage;
+	ITypedElementDataStorageUiInterface& StorageUi;
+	ITypedElementDataStorageCompatibilityInterface& StorageCompatibility;
+	FName NameId;
+};
+
+
+
+
+//
+// UTypedElementSceneOutlinerFactory
+// 
+void UTypedElementSceneOutlinerFactory::RegisterTables(ITypedElementDataStorageInterface& DataStorage) const
+{
+	TypedElementTableHandle BaseWidgetTable = DataStorage.FindTable(FName(TEXT("Editor_WidgetTable")));
+	if (BaseWidgetTable != TypedElementInvalidTableHandle)
+	{
+		DataStorage.RegisterTable(
+			BaseWidgetTable,
+			{
+				FTypedElementRowReferenceColumn::StaticStruct()
+			},
+			FTypedElementSceneOutlinerQueryBinder::CellWidgetTableName);
+	}
+}
+
+void UTypedElementSceneOutlinerFactory::RegisterWidgetPurposes(ITypedElementDataStorageUiInterface& DataStorageUi) const
+{
+	using PurposeType = ITypedElementDataStorageUiInterface::EPurposeType;
+
+	DataStorageUi.RegisterWidgetPurpose(FTypedElementSceneOutlinerQueryBinder::HeaderWidgetPurpose, PurposeType::UniqueByNameAndColumn,
+		LOCTEXT("SceneOutlinerCellDescription", "Widgets for headers in any Scene Outliner for specific columns or column combinations."));
+	DataStorageUi.RegisterWidgetPurpose(FTypedElementSceneOutlinerQueryBinder::DefaultHeaderWidgetPurpose, PurposeType::UniqueByName,
+		LOCTEXT("SceneOutlinerCellDescription", "The default widget to use in headers for the Scene Outliner."));
+	
+	DataStorageUi.RegisterWidgetPurpose(FTypedElementSceneOutlinerQueryBinder::CellWidgetPurpose, PurposeType::UniqueByNameAndColumn,
+		LOCTEXT("SceneOutlinerCellDescription", "Widgets for cells in any Scene Outliner for specific columns or column combinations."));
+	DataStorageUi.RegisterWidgetPurpose(FTypedElementSceneOutlinerQueryBinder::DefaultCellWidgetPurpose, PurposeType::UniqueByName,
+		LOCTEXT("SceneOutlinerCellDescription", "The default widget to use in cells for the Scene Outliner."));
+}
+
+
+
+//
+// FTypedElementSceneOutlinerQueryBinder
+// 
+
+const FName FTypedElementSceneOutlinerQueryBinder::CellWidgetTableName(TEXT("Editor_SceneOutlinerCellWidgetTable"));
+const FName FTypedElementSceneOutlinerQueryBinder::HeaderWidgetPurpose(TEXT("SceneOutliner.Header"));
+const FName FTypedElementSceneOutlinerQueryBinder::DefaultHeaderWidgetPurpose(TEXT("SceneOutliner.Header.Default"));
+const FName FTypedElementSceneOutlinerQueryBinder::CellWidgetPurpose(TEXT("SceneOutliner.Cell"));
+const FName FTypedElementSceneOutlinerQueryBinder::DefaultCellWidgetPurpose(TEXT("SceneOutliner.Cell.Default"));
+
+FTypedElementSceneOutlinerQueryBinder::FTypedElementSceneOutlinerQueryBinder()
+{
+	UTypedElementRegistry* Registry = UTypedElementRegistry::GetInstance();
+	checkf(Registry, TEXT("Unable to bind a Scene Outliner to a query before the Typed Elements are available."));
+	if (Registry)
+	{
+		Storage = Registry->GetMutableDataStorage();
+		StorageUi = Registry->GetMutableDataStorageUi();
+		StorageCompatibility = Registry->GetMutableDataStorageCompatibility();
+	}
+}
+
+FTypedElementSceneOutlinerQueryBinder& FTypedElementSceneOutlinerQueryBinder::GetInstance()
+{
+	static FTypedElementSceneOutlinerQueryBinder Binder;
+	return Binder;
+}
+
+void FTypedElementSceneOutlinerQueryBinder::AssignQuery(TypedElementQueryHandle Query, const TSharedPtr<ISceneOutliner>& Outliner)
+{
+	FTypedElementSceneOutliner* QueryMapping = SceneOutliners.Find(Outliner);
+	if (QueryMapping == nullptr)
+	{
+		QueryMapping = &SceneOutliners.Add(Outliner);
+		QueryMapping->Initialize(*Storage, *StorageUi, *StorageCompatibility, Outliner);
+	}
+	QueryMapping->AssignQuery(Query);
+}
+
+
+
+//
+// FTypedElementSceneOutliner
+//
+
+FTypedElementSceneOutliner::~FTypedElementSceneOutliner()
+{
+	TSharedPtr<ISceneOutliner> OutlinerPinned = Outliner.Pin();
+	if (OutlinerPinned)
+	{
+		ClearColumns(*OutlinerPinned);
+	}
+}
+
+void FTypedElementSceneOutliner::Initialize(
+	ITypedElementDataStorageInterface& InStorage,
+	ITypedElementDataStorageUiInterface& InStorageUi,
+	ITypedElementDataStorageCompatibilityInterface& InStorageCompatibility,
+	const TSharedPtr<ISceneOutliner>& InOutliner)
+{
+	Storage = &InStorage;
+	StorageUi = &InStorageUi;
+	StorageCompatibility = &InStorageCompatibility;
+	Outliner = InOutliner;
+}
+
+TArray<TWeakObjectPtr<const UScriptStruct>> FTypedElementSceneOutliner::CreateVerifiedColumnTypeAray(
+	TConstArrayView<TWeakObjectPtr<const UScriptStruct>> ColumnTypes)
+{
+	TArray<TWeakObjectPtr<const UScriptStruct>> VerifiedColumnTypes;
+	VerifiedColumnTypes.Reserve(ColumnTypes.Num());
+	for (const TWeakObjectPtr<const UScriptStruct>& ColumnType : ColumnTypes)
+	{
+		if (ColumnType.IsValid())
+		{
+			VerifiedColumnTypes.Add(ColumnType.Get());
+		}
+	}
+	return VerifiedColumnTypes;
+}
+
+TSharedPtr<FTypedElementWidgetConstructor> FTypedElementSceneOutliner::CreateHeaderWidgetConstructor(
+	ITypedElementDataStorageUiInterface& StorageUi,	TConstArrayView<TWeakObjectPtr<const UScriptStruct>> ColumnTypes)
+{
+	using MatchApproach = ITypedElementDataStorageUiInterface::EMatchApproach;
+
+	TArray<TWeakObjectPtr<const UScriptStruct>> VerifiedColumnTypes = CreateVerifiedColumnTypeAray(ColumnTypes);
+	TSharedPtr<FTypedElementWidgetConstructor> Constructor;
+
+	FName HeaderNames[] = { TEXT("SceneOutliner.Header"), TEXT("General.Header") };
+	for (const FName& Purpose : HeaderNames)
+	{
+		StorageUi.CreateWidgetConstructors(Purpose, MatchApproach::ExactMatch, VerifiedColumnTypes, {}, 
+			[&Constructor, ColumnTypes](
+				TUniquePtr<FTypedElementWidgetConstructor> CreatedConstructor, 
+				TConstArrayView<TWeakObjectPtr<const UScriptStruct>> MatchedColumnTypes)
+			{
+				if (ColumnTypes.Num() == MatchedColumnTypes.Num())
+				{
+					Constructor = TSharedPtr<FTypedElementWidgetConstructor>(CreatedConstructor.Release());
+				}
+				// Either this was the exact match so no need to search further or the longest possible chain didn't match so the next ones will 
+				// always be shorter in both cases just return.
+				return false;
+			});
+		if (Constructor)
+		{
+			return Constructor;
+		}
+	}
+	FName HeaderDefaultNames[] = { TEXT("SceneOutliner.Header.Default"), TEXT("General.Header.Default") };
+	for (const FName& Purpose : HeaderDefaultNames)
+	{
+		StorageUi.CreateWidgetConstructors(Purpose, {}, 
+			[&Constructor, ColumnTypes](
+				TUniquePtr<FTypedElementWidgetConstructor> CreatedConstructor,
+				TConstArrayView<TWeakObjectPtr<const UScriptStruct>> MatchedColumnTypes)
+			{
+				Constructor = TSharedPtr<FTypedElementWidgetConstructor>(CreatedConstructor.Release());
+				return false;
+			});
+		if (Constructor)
+		{
+			return Constructor;
+		}
+	}
+	return nullptr;
+}
+
+void FTypedElementSceneOutliner::AssignQuery(TypedElementQueryHandle Query)
+{
+	using MatchApproach = ITypedElementDataStorageUiInterface::EMatchApproach;
+	constexpr int32 DefaultPriorityIndex = 100;
+
+	const ITypedElementDataStorageInterface::FQueryDescription& Description = Storage->GetQueryDescription(Query);
+	if (TSharedPtr<ISceneOutliner> OutlinerPinned = Outliner.Pin())
+	{
+		ClearColumns(*OutlinerPinned);
+
+		if (Description.Action == ITypedElementDataStorageInterface::FQueryDescription::EActionType::Select)
+		{
+			int32 SelectionCount = Description.SelectionTypes.Num();
+			AddedColumns.Reset(SelectionCount);
+
+			TArray<TWeakObjectPtr<const UScriptStruct>> ColumnTypes = CreateVerifiedColumnTypeAray(Description.SelectionTypes);
+			
+			int32 IndexOffset = 0;
+			auto ColumnConstructor = [this, &IndexOffset, &OutlinerPinned](
+				TUniquePtr<FTypedElementWidgetConstructor> Constructor, TConstArrayView<TWeakObjectPtr<const UScriptStruct>> ColumnTypes)
+				{
+					TSharedPtr<FTypedElementWidgetConstructor> CellConstructor(Constructor.Release());
+			
+					FName NameId = FindLongestMatchingName(ColumnTypes, IndexOffset);
+					AddedColumns.Add(NameId);
+					OutlinerPinned->AddColumn(NameId,
+						FSceneOutlinerColumnInfo(ESceneOutlinerColumnVisibility::Visible, DefaultPriorityIndex + IndexOffset,
+							FCreateSceneOutlinerColumn::CreateLambda(
+								[this, NameId, &ColumnTypes, CellConstructor](ISceneOutliner&)
+								{
+									TSharedPtr<FTypedElementWidgetConstructor> HeaderConstructor = 
+										CreateHeaderWidgetConstructor(*StorageUi, ColumnTypes);
+									return MakeShared<FOutlinerColumn>(
+										*Storage, *StorageUi, *StorageCompatibility, NameId,
+										TArray<TWeakObjectPtr<const UScriptStruct>>(ColumnTypes.GetData(), ColumnTypes.Num()), 
+										MoveTemp(HeaderConstructor), CellConstructor);
+
+								})
+						)
+					);
+					++IndexOffset;
+					return true;
+				};
+			StorageUi->CreateWidgetConstructors(FName(TEXT("SceneOutliner.Cell")), MatchApproach::LongestMatch, ColumnTypes, {}, ColumnConstructor);
+			StorageUi->CreateWidgetConstructors(FName(TEXT("General.Cell")), MatchApproach::LongestMatch, ColumnTypes, {}, ColumnConstructor);
+
+			for (TWeakObjectPtr<const UScriptStruct>& ColumnType : ColumnTypes)
+			{
+				auto AssignWidgetToColumn = [this, ColumnType, &IndexOffset, &OutlinerPinned](
+					TUniquePtr<FTypedElementWidgetConstructor> Constructor, TConstArrayView<TWeakObjectPtr<const UScriptStruct>>)
+				{
+					TSharedPtr<FTypedElementWidgetConstructor> CellConstructor(Constructor.Release());
+					FName NameId = FName(ColumnType->GetDisplayNameText().ToString());
+					AddedColumns.Add(NameId);
+					OutlinerPinned->AddColumn(NameId,
+						FSceneOutlinerColumnInfo(ESceneOutlinerColumnVisibility::Visible, DefaultPriorityIndex + IndexOffset,
+							FCreateSceneOutlinerColumn::CreateLambda(
+								[this, NameId, ColumnType, CellConstructor](ISceneOutliner&)
+								{
+									TArray<TWeakObjectPtr<const UScriptStruct>> ColumnTypesStored;
+									ColumnTypesStored.Add(ColumnType);
+									TSharedPtr<FTypedElementWidgetConstructor> HeaderConstructor =
+										CreateHeaderWidgetConstructor(*StorageUi, { ColumnType });
+									return MakeShared<FOutlinerColumn>(
+										*Storage, *StorageUi, *StorageCompatibility, NameId, MoveTemp(ColumnTypesStored), 
+										HeaderConstructor, CellConstructor);
+
+								})
+						)
+					);
+					++IndexOffset;
+					return false;
+				};
+				int32 BeforeIndexOffset = IndexOffset;
+				StorageUi->CreateWidgetConstructors(FName(TEXT("SceneOutliner.Cell.Default")), {}, AssignWidgetToColumn);
+				if (BeforeIndexOffset == IndexOffset)
+				{
+					StorageUi->CreateWidgetConstructors(FName(TEXT("General.Cell.Default")), {}, AssignWidgetToColumn);
+					if (BeforeIndexOffset == IndexOffset)
+					{
+						++IndexOffset;
+					}
+				}
+			}
+		}
+	}
+}
+
+FName FTypedElementSceneOutliner::FindLongestMatchingName(TConstArrayView<TWeakObjectPtr<const UScriptStruct>> ColumnTypes, int32 DefaultNameIndex)
+{
+	switch (ColumnTypes.Num())
+	{
+	case 0:
+		return FName(TEXT("Column"), DefaultNameIndex);
+	case 1:
+		return FName(ColumnTypes[0]->GetDisplayNameText().ToString());
+	default:
+	{
+		FString LongestMatch(ColumnTypes[0]->GetDisplayNameText().ToString());
+		const TWeakObjectPtr<const UScriptStruct>* ItEnd = ColumnTypes.end();
+		const TWeakObjectPtr<const UScriptStruct>* It = ColumnTypes.begin();
+		++It; // Skip the first entry as that's already set.
+		for (; It != ItEnd; ++It)
+		{
+			const FString& NextMatch = (*It)->GetDisplayNameText().ToString();
+
+			int32 MatchSize = 0;
+			auto ItLeft = LongestMatch.begin();
+			auto ItLeftEnd = LongestMatch.end();
+			auto ItRight = NextMatch.begin();
+			auto ItRightEnd = NextMatch.end();
+			while (
+				ItLeft != ItLeftEnd &&
+				ItRight != ItRightEnd &&
+				*ItLeft == *ItRight)
+			{
+				++MatchSize;
+				++ItLeft;
+				++ItRight;
+			}
+
+			// At least 3 letters have to match to avoid single or double letter names which typically mean nothing.
+			if (MatchSize > 2)
+			{
+				LongestMatch.LeftInline(MatchSize);
+			}
+			else
+			{
+				// There are not enough characters in the string that match.
+				return FName(TEXT("Column"), DefaultNameIndex);
+			}
+		}
+		return FName(MoveTemp(LongestMatch));
+	}
+	};
+}
+
+void FTypedElementSceneOutliner::ClearColumns(ISceneOutliner& InOutliner)
+{
+	for (FName ColumnName : AddedColumns)
+	{
+		InOutliner.RemoveColumn(ColumnName);
+	}
+}
+
+#undef LOCTEXT_NAMESPACE
\ No newline at end of file
diff --git a/Engine/Plugins/Experimental/TypedElementsDataStorage/Source/TypedElementsDataStorageUI/Public/TypedElementExportedTextWidget.h b/Engine/Plugins/Experimental/TypedElementsDataStorage/Source/TypedElementsDataStorageUI/Public/TypedElementExportedTextWidget.h
new file mode 100644
--- /dev/null
+++ b/Engine/Plugins/Experimental/TypedElementsDataStorage/Source/TypedElementsDataStorageUI/Public/TypedElementExportedTextWidget.h
@@ -0,0 +1,51 @@
+// Copyright Epic Games, Inc. All Rights Reserved.
+
+#pragma once
+
+#include "Elements/Interfaces/TypedElementDataStorageFactory.h"
+#include "Elements/Interfaces/TypedElementDataStorageInterface.h"
+#include "Elements/Interfaces/TypedElementDataStorageUiInterface.h"
+#include "Internationalization/Text.h"
+#include "UObject/ObjectMacros.h"
+
+#include "TypedElementExportedTextWidget.generated.h"
+
+UCLASS()
+class TYPEDELEMENTSDATASTORAGEUI_API UTypedElementExportedTextWidgetFactory : public UTypedElementDataStorageFactory
+{
+	GENERATED_BODY()
+
+public:
+	~UTypedElementExportedTextWidgetFactory() override = default;
+
+	void RegisterQueries(ITypedElementDataStorageInterface& DataStorage) const override;
+	void RegisterWidgetConstructors(ITypedElementDataStorageInterface& DataStorage,
+		ITypedElementDataStorageUiInterface& DataStorageUi) const override;
+};
+
+USTRUCT()
+struct TYPEDELEMENTSDATASTORAGEUI_API FTypedElementExportedTextWidgetConstructor : public FTypedElementWidgetConstructor
+{
+	GENERATED_BODY()
+
+public:
+	FTypedElementExportedTextWidgetConstructor();
+	~FTypedElementExportedTextWidgetConstructor() override = default;
+
+	TConstArrayView<const UScriptStruct*> GetAdditionalColumnsList() const override;
+	bool CanBeReused() const override;
+
+protected:
+	TSharedPtr<SWidget> CreateWidget() override;
+	bool FinalizeWidget(
+		ITypedElementDataStorageInterface* DataStorage,
+		ITypedElementDataStorageUiInterface* DataStorageUi,
+		TypedElementRowHandle Row,
+		const TSharedPtr<SWidget>& Widget) override;
+};
+
+USTRUCT(meta = (DisplayName = "Exported text widget"))
+struct TYPEDELEMENTSDATASTORAGEUI_API FTypedElementExportedTextWidgetTag : public FTypedElementDataStorageTag
+{
+	GENERATED_BODY()
+};
\ No newline at end of file
diff --git a/Engine/Plugins/Experimental/TypedElementsDataStorage/Source/TypedElementsDataStorageUI/Public/TypedElementGeneralRegistration.h b/Engine/Plugins/Experimental/TypedElementsDataStorage/Source/TypedElementsDataStorageUI/Public/TypedElementGeneralRegistration.h
new file mode 100644
--- /dev/null
+++ b/Engine/Plugins/Experimental/TypedElementsDataStorage/Source/TypedElementsDataStorageUI/Public/TypedElementGeneralRegistration.h
@@ -0,0 +1,23 @@
+// Copyright Epic Games, Inc. All Rights Reserved.
+
+#pragma once
+
+#include "Elements/Interfaces/TypedElementDataStorageFactory.h"
+#include "UObject/ObjectMacros.h"
+
+#include "TypedElementGeneralRegistration.generated.h"
+
+UCLASS()
+class TYPEDELEMENTSDATASTORAGEUI_API UTypedElementGeneralRegistrationFactory : public UTypedElementDataStorageFactory
+{
+	GENERATED_BODY()
+
+public:
+	static const FName CellPurpose;
+	static const FName CellDefaultPurpose;
+	static const FName HeaderDefaultPurpose;
+
+	~UTypedElementGeneralRegistrationFactory() override = default;
+
+	void RegisterWidgetPurposes(ITypedElementDataStorageUiInterface& DataStorageUi) const override;
+};
diff --git a/Engine/Plugins/Experimental/TypedElementsDataStorage/Source/TypedElementsDataStorageUI/Public/TypedElementSceneOutliner.h b/Engine/Plugins/Experimental/TypedElementsDataStorage/Source/TypedElementsDataStorageUI/Public/TypedElementSceneOutliner.h
new file mode 100644
--- /dev/null
+++ b/Engine/Plugins/Experimental/TypedElementsDataStorage/Source/TypedElementsDataStorageUI/Public/TypedElementSceneOutliner.h
@@ -0,0 +1,85 @@
+// Copyright Epic Games, Inc. All Rights Reserved.
+
+#pragma once
+
+#include "Containers/ContainersFwd.h"
+#include "Elements/Interfaces/TypedElementDataStorageFactory.h"
+#include "Elements/Interfaces/TypedElementDataStorageInterface.h"
+#include "Elements/Interfaces/TypedElementDataStorageUiInterface.h"
+#include "Templates/SharedPointer.h"
+#include "UObject/ObjectMacros.h"
+
+#include "TypedElementSceneOutliner.generated.h"
+
+class ISceneOutliner;
+class ITypedElementDataStorageInterface;
+class ITypedElementDataStorageUiInterface;
+class ITypedElementDataStorageCompatibilityInterface;
+
+class FTypedElementSceneOutliner
+{
+public:
+	~FTypedElementSceneOutliner();
+
+	void Initialize(
+		ITypedElementDataStorageInterface& InStorage,
+		ITypedElementDataStorageUiInterface& InStorageUi,
+		ITypedElementDataStorageCompatibilityInterface& InStorageCompatibility,
+		const TSharedPtr<ISceneOutliner>& InOutliner);
+
+	void AssignQuery(TypedElementQueryHandle Query);
+
+private:
+	static FName FindLongestMatchingName(TConstArrayView<TWeakObjectPtr<const UScriptStruct>> ColumnTypes, int32 DefaultNameIndex);
+	static TArray<TWeakObjectPtr<const UScriptStruct>> CreateVerifiedColumnTypeAray(
+		TConstArrayView<TWeakObjectPtr<const UScriptStruct>> ColumnTypes);
+	static TSharedPtr<FTypedElementWidgetConstructor> CreateHeaderWidgetConstructor(
+		ITypedElementDataStorageUiInterface& StorageUI, TConstArrayView<TWeakObjectPtr<const UScriptStruct>> ColumnTypes);
+	void ClearColumns(ISceneOutliner& InOutliner);
+
+	TArray<FName> AddedColumns;
+	TWeakPtr<ISceneOutliner> Outliner;
+	ITypedElementDataStorageInterface* Storage{ nullptr };
+	ITypedElementDataStorageUiInterface* StorageUi{ nullptr };
+	ITypedElementDataStorageCompatibilityInterface* StorageCompatibility{ nullptr };
+};
+
+/**
+ * Utility class to bind Typed Elements Data Storage queries to a Scene Outliner. The provided query is expected to be a select query
+ * and will be used to populate the Scene Outliner in addition to already existing data.
+ */
+class FTypedElementSceneOutlinerQueryBinder
+{
+public:
+	static const FName CellWidgetTableName;
+	static const FName HeaderWidgetPurpose;
+	static const FName DefaultHeaderWidgetPurpose;
+	static const FName CellWidgetPurpose;
+	static const FName DefaultCellWidgetPurpose;
+
+	static FTypedElementSceneOutlinerQueryBinder& GetInstance();
+
+	void AssignQuery(TypedElementQueryHandle Query, const TSharedPtr<ISceneOutliner>& Widget);
+
+private:
+	FTypedElementSceneOutlinerQueryBinder();
+
+	TMap<TWeakPtr<ISceneOutliner>, FTypedElementSceneOutliner> SceneOutliners;
+
+	ITypedElementDataStorageInterface* Storage{ nullptr };
+	ITypedElementDataStorageUiInterface* StorageUi{ nullptr };
+	ITypedElementDataStorageCompatibilityInterface* StorageCompatibility{ nullptr };
+};
+
+UCLASS()
+class UTypedElementSceneOutlinerFactory : public UTypedElementDataStorageFactory
+{
+	GENERATED_BODY()
+
+public:
+	~UTypedElementSceneOutlinerFactory() override = default;
+
+	void RegisterTables(ITypedElementDataStorageInterface& DataStorage) const override;
+	
+	void RegisterWidgetPurposes(ITypedElementDataStorageUiInterface& DataStorageUi) const override;
+};
diff --git a/Engine/Plugins/Experimental/TypedElementsDataStorage/Source/TypedElementsDataStorageUI/TypedElementsDataStorageUI.Build.cs b/Engine/Plugins/Experimental/TypedElementsDataStorage/Source/TypedElementsDataStorageUI/TypedElementsDataStorageUI.Build.cs
--- a/Engine/Plugins/Experimental/TypedElementsDataStorage/Source/TypedElementsDataStorageUI/TypedElementsDataStorageUI.Build.cs
+++ b/Engine/Plugins/Experimental/TypedElementsDataStorage/Source/TypedElementsDataStorageUI/TypedElementsDataStorageUI.Build.cs
@@ -20,6 +20,7 @@ public class TypedElementsDataStorageUI : ModuleRules
 					"CoreUObject",
 					"Engine",
 					"EditorFramework",
+					"SceneOutliner",
 					"Slate",
 					"SlateCore",
 					"ToolMenus",
