Fixed missing license header.

#nojira

diff --git a/Engine/Plugins/Experimental/TypedElementsDataStorage/Source/TypedElementsDataStorageUI/Private/TypedElementExportedTextWidget.cpp b/Engine/Plugins/Experimental/TypedElementsDataStorage/Source/TypedElementsDataStorageUI/Private/TypedElementExportedTextWidget.cpp
--- a/Engine/Plugins/Experimental/TypedElementsDataStorage/Source/TypedElementsDataStorageUI/Private/TypedElementExportedTextWidget.cpp
+++ b/Engine/Plugins/Experimental/TypedElementsDataStorage/Source/TypedElementsDataStorageUI/Private/TypedElementExportedTextWidget.cpp
@@ -1,3 +1,5 @@
+// Copyright Epic Games, Inc. All Rights Reserved.
+
 #include "TypedElementExportedTextWidget.h"
 
 #include "Elements/Columns/TypedElementMiscColumns.h"
