Fixed missing license header

#nojira

diff --git a/Engine/Plugins/Experimental/TypedElementsDataStorage/Source/TypedElementsDataStorageUI/Private/TypedElementGeneralRegistration.cpp b/Engine/Plugins/Experimental/TypedElementsDataStorage/Source/TypedElementsDataStorageUI/Private/TypedElementGeneralRegistration.cpp
--- a/Engine/Plugins/Experimental/TypedElementsDataStorage/Source/TypedElementsDataStorageUI/Private/TypedElementGeneralRegistration.cpp
+++ b/Engine/Plugins/Experimental/TypedElementsDataStorage/Source/TypedElementsDataStorageUI/Private/TypedElementGeneralRegistration.cpp
@@ -1,3 +1,5 @@
+// Copyright Epic Games, Inc. All Rights Reserved.
+
 #include "TypedElementGeneralRegistration.h"
 
 #include "Elements/Interfaces/TypedElementDataStorageUiInterface.h"
