Moved handling of queries from the Data Storage class to its own class.

This moves most of the query logic out of the Data Storage class and into its own class so the same logic can be shared between the Data Storage (which handles running queries directly) and the query context (which will be handling subqueries in the future). Other than some additional error handling, there are no functional differences as code has been moved and cleaned up.

diff --git a/Engine/Plugins/Experimental/TypedElementsDataStorage/Source/TypedElementsDataStorage/Private/Processors/TypedElementProcessorAdaptors.cpp b/Engine/Plugins/Experimental/TypedElementsDataStorage/Source/TypedElementsDataStorage/Private/Processors/TypedElementProcessorAdaptors.cpp
--- a/Engine/Plugins/Experimental/TypedElementsDataStorage/Source/TypedElementsDataStorage/Private/Processors/TypedElementProcessorAdaptors.cpp
+++ b/Engine/Plugins/Experimental/TypedElementsDataStorage/Source/TypedElementsDataStorage/Private/Processors/TypedElementProcessorAdaptors.cpp
@@ -150,7 +150,7 @@ struct FMassContextForwarder final : public ITypedElementDataStorageInterface::I
 		}
 	}
 
-	void AddColumns(TypedElementRowHandle Row, TConstArrayView<const UScriptStruct*> ColumnTypes)
+	void AddColumns(TypedElementRowHandle Row, TConstArrayView<const UScriptStruct*> ColumnTypes) override
 	{
 		for (const UScriptStruct* ColumnType : ColumnTypes)
 		{
@@ -178,7 +178,7 @@ struct FMassContextForwarder final : public ITypedElementDataStorageInterface::I
 		}
 	}
 
-	void AddColumns(TConstArrayView<TypedElementRowHandle> Rows, TConstArrayView<const UScriptStruct*> ColumnTypes)
+	void AddColumns(TConstArrayView<TypedElementRowHandle> Rows, TConstArrayView<const UScriptStruct*> ColumnTypes) override
 	{
 		for (TypedElementRowHandle Row : Rows)
 		{
@@ -186,7 +186,7 @@ struct FMassContextForwarder final : public ITypedElementDataStorageInterface::I
 		}
 	}
 
-	void RemoveColumns(TypedElementRowHandle Row, TConstArrayView<const UScriptStruct*> ColumnTypes)
+	void RemoveColumns(TypedElementRowHandle Row, TConstArrayView<const UScriptStruct*> ColumnTypes) override
 	{
 		for (const UScriptStruct* ColumnType : ColumnTypes)
 		{
@@ -214,7 +214,7 @@ struct FMassContextForwarder final : public ITypedElementDataStorageInterface::I
 		}
 	}
 
-	void RemoveColumns(TConstArrayView<TypedElementRowHandle> Rows, TConstArrayView<const UScriptStruct*> ColumnTypes)
+	void RemoveColumns(TConstArrayView<TypedElementRowHandle> Rows, TConstArrayView<const UScriptStruct*> ColumnTypes) override
 	{
 		for (TypedElementRowHandle Row : Rows)
 		{
@@ -427,7 +427,7 @@ FMassEntityQuery& UTypedElementQueryProcessorCallbackAdapterProcessor::GetQuery(
 	return Data.Query;
 }
 
-void UTypedElementQueryProcessorCallbackAdapterProcessor::ConfigureQueryCallback(FTypedElementDatabaseExtendedQuery& TargetParentQuery)
+void UTypedElementQueryProcessorCallbackAdapterProcessor::ConfigureQueryCallback(FTypedElementExtendedQuery& TargetParentQuery)
 {
 	Data.ParentQuery = &TargetParentQuery;
 
@@ -491,7 +491,7 @@ EMassObservedOperation UTypedElementQueryObserverCallbackAdapterProcessor::GetOb
 	return Operation;
 }
 
-void UTypedElementQueryObserverCallbackAdapterProcessor::ConfigureQueryCallback(FTypedElementDatabaseExtendedQuery& TargetParentQuery)
+void UTypedElementQueryObserverCallbackAdapterProcessor::ConfigureQueryCallback(FTypedElementExtendedQuery& TargetParentQuery)
 {
 	Data.ParentQuery = &TargetParentQuery;
 
diff --git a/Engine/Plugins/Experimental/TypedElementsDataStorage/Source/TypedElementsDataStorage/Private/Processors/TypedElementProcessorAdaptors.h b/Engine/Plugins/Experimental/TypedElementsDataStorage/Source/TypedElementsDataStorage/Private/Processors/TypedElementProcessorAdaptors.h
--- a/Engine/Plugins/Experimental/TypedElementsDataStorage/Source/TypedElementsDataStorage/Private/Processors/TypedElementProcessorAdaptors.h
+++ b/Engine/Plugins/Experimental/TypedElementsDataStorage/Source/TypedElementsDataStorage/Private/Processors/TypedElementProcessorAdaptors.h
@@ -11,7 +11,7 @@
 
 #include "TypedElementProcessorAdaptors.generated.h"
 
-struct FTypedElementDatabaseExtendedQuery;
+struct FTypedElementExtendedQuery;
 
 struct FPhasePreOrPostAmbleExecutor
 {
@@ -45,7 +45,7 @@ struct FTypedElementQueryProcessorData
 	static bool PrepareCachedDependenciesOnQuery(
 		ITypedElementDataStorageInterface::FQueryDescription& Description, FMassExecutionContext& Context);
 
-	FTypedElementDatabaseExtendedQuery* ParentQuery{ nullptr };
+	FTypedElementExtendedQuery* ParentQuery{ nullptr };
 	FMassEntityQuery Query;
 };
 
@@ -61,7 +61,7 @@ public:
 	UTypedElementQueryProcessorCallbackAdapterProcessor();
 
 	FMassEntityQuery& GetQuery();
-	void ConfigureQueryCallback(FTypedElementDatabaseExtendedQuery& Query);
+	void ConfigureQueryCallback(FTypedElementExtendedQuery& Query);
 
 protected:
 	void ConfigureQueries() override;
@@ -89,7 +89,7 @@ public:
 	FMassEntityQuery& GetQuery();
 	const UScriptStruct* GetObservedType() const;
 	EMassObservedOperation GetObservedOperation() const;
-	void ConfigureQueryCallback(FTypedElementDatabaseExtendedQuery& Query);
+	void ConfigureQueryCallback(FTypedElementExtendedQuery& Query);
 
 protected:
 	void ConfigureQueries() override;
diff --git a/Engine/Plugins/Experimental/TypedElementsDataStorage/Source/TypedElementsDataStorage/Private/Queries/TypedElementExtendedQueryStore.cpp b/Engine/Plugins/Experimental/TypedElementsDataStorage/Source/TypedElementsDataStorage/Private/Queries/TypedElementExtendedQueryStore.cpp
new file mode 100644
--- /dev/null
+++ b/Engine/Plugins/Experimental/TypedElementsDataStorage/Source/TypedElementsDataStorage/Private/Queries/TypedElementExtendedQueryStore.cpp
@@ -0,0 +1,576 @@
+// Copyright Epic Games, Inc. All Rights Reserved.
+
+#include "Queries/TypedElementExtendedQueryStore.h"
+
+#include "MassProcessingPhaseManager.h"
+#include "MassProcessor.h"
+#include "Misc/OutputDevice.h"
+#include "Processors/TypedElementProcessorAdaptors.h"
+
+const ITypedElementDataStorageInterface::FQueryDescription FTypedElementExtendedQueryStore::EmptyDescription{};
+
+FTypedElementExtendedQueryStore::Handle FTypedElementExtendedQueryStore::RegisterQuery(
+	ITypedElementDataStorageInterface::FQueryDescription Query, FMassEntityManager& EntityManager, FMassProcessingPhaseManager& PhaseManager)
+{
+	FTypedElementExtendedQueryStore::Handle Result = Queries.Emplace();
+	FTypedElementExtendedQuery& StoredQuery = GetMutableChecked(Result);
+	StoredQuery.Description = MoveTemp(Query);
+
+	FMassEntityQuery& NativeQuery =		SetupNativeQuery(StoredQuery.Description, StoredQuery);
+	bool bContinueSetup =				SetupSelectedColumns(StoredQuery.Description, NativeQuery);
+	bContinueSetup = bContinueSetup &&	SetupConditions(StoredQuery.Description, NativeQuery);
+	bContinueSetup = bContinueSetup &&	SetupDependencies(StoredQuery.Description, NativeQuery);
+	bContinueSetup = bContinueSetup &&	SetupTickGroupDefaults(StoredQuery.Description);
+	bContinueSetup = bContinueSetup &&	SetupProcessors(Result, StoredQuery, EntityManager, PhaseManager);
+	
+	if (!bContinueSetup)
+	{
+		// This will also make the handle invalid.
+		Queries.Remove(Result);
+	}
+
+	return Result;
+}
+
+void FTypedElementExtendedQueryStore::UnregisterQuery(Handle Query, FMassProcessingPhaseManager& PhaseManager)
+{
+	if (FTypedElementExtendedQuery* QueryData = Get(Query))
+	{
+		if (QueryData->Processor)
+		{
+			if (QueryData->Processor->IsA<UTypedElementQueryProcessorCallbackAdapterProcessor>())
+			{
+				PhaseManager.UnregisterDynamicProcessor(*QueryData->Processor);
+			}
+			else if (QueryData->Processor->IsA<UTypedElementQueryObserverCallbackAdapterProcessor>())
+			{
+				checkf(false, TEXT("Observer queries can not be unregistered."));
+			}
+			else
+			{
+				checkf(false, TEXT("Query processor %s is of unsupported type %s."),
+					*QueryData->Description.Callback.Name.ToString(), *QueryData->Processor->GetSparseClassDataStruct()->GetName());
+			}
+		}
+		else if (QueryData->Description.Callback.Type == ITypedElementDataStorageInterface::EQueryCallbackType::PhasePreparation)
+		{
+			UnregisterPreambleQuery(QueryData->Description.Callback.Phase, Query);
+		}
+		else if (QueryData->Description.Callback.Type == ITypedElementDataStorageInterface::EQueryCallbackType::PhaseFinalization)
+		{
+			UnregisterPostambleQuery(QueryData->Description.Callback.Phase, Query);
+		}
+		else
+		{
+			QueryData->NativeQuery.Clear();
+		}
+		Queries.Remove(Query);
+	}
+}
+
+void FTypedElementExtendedQueryStore::RegisterTickGroup(FName GroupName, ITypedElementDataStorageInterface::EQueryTickPhase Phase,
+	FName BeforeGroup, FName AfterGroup, bool bRequiresMainThread)
+{
+	FTickGroupDescription& Group = TickGroupDescriptions.FindOrAdd({ GroupName, Phase });
+
+	if (!Group.BeforeGroups.Find(BeforeGroup))
+	{
+		Group.BeforeGroups.Add(BeforeGroup);
+	}
+
+	if (!Group.AfterGroups.Find(AfterGroup))
+	{
+		Group.AfterGroups.Add(AfterGroup);
+	}
+
+	if (bRequiresMainThread)
+	{
+		Group.bRequiresMainThread = true;
+	}
+}
+
+void FTypedElementExtendedQueryStore::UnregisterTickGroup(FName GroupName, ITypedElementDataStorageInterface::EQueryTickPhase Phase)
+{
+	TickGroupDescriptions.Remove({ GroupName, Phase });
+}
+
+FTypedElementExtendedQuery* FTypedElementExtendedQueryStore::Get(Handle Entry)
+{
+	return IsAlive(Entry) ? &Queries.Get(Entry) : nullptr;
+}
+
+FTypedElementExtendedQuery* FTypedElementExtendedQueryStore::GetMutable(Handle Entry)
+{
+	return IsAlive(Entry) ? &Queries.GetMutable(Entry) : nullptr;
+}
+
+const FTypedElementExtendedQuery* FTypedElementExtendedQueryStore::Get(Handle Entry) const
+{
+	return IsAlive(Entry) ? &Queries.Get(Entry) : nullptr;
+}
+
+FTypedElementExtendedQuery& FTypedElementExtendedQueryStore::GetChecked(Handle Entry)
+{
+	return Queries.Get(Entry);
+}
+
+FTypedElementExtendedQuery& FTypedElementExtendedQueryStore::GetMutableChecked(Handle Entry)
+{
+	return Queries.GetMutable(Entry);
+}
+
+const FTypedElementExtendedQuery& FTypedElementExtendedQueryStore::GetChecked(Handle Entry) const
+{
+	return Queries.Get(Entry);
+}
+
+const ITypedElementDataStorageInterface::FQueryDescription& FTypedElementExtendedQueryStore::GetQueryDescription(Handle Query) const
+{
+	const FTypedElementExtendedQuery* QueryData = Get(Query);
+	return QueryData ? QueryData->Description : EmptyDescription;
+}
+
+bool FTypedElementExtendedQueryStore::IsAlive(Handle Entry) const
+{
+	return Queries.IsAlive(Entry);
+}
+
+void FTypedElementExtendedQueryStore::ListAliveEntries(const ListAliveEntriesCallback& Callback) const
+{
+	Queries.ListAliveEntries(Callback);
+}
+
+ITypedElementDataStorageInterface::FQueryResult FTypedElementExtendedQueryStore::RunQuery(FMassEntityManager& EntityManager, Handle Query)
+{
+	using ActionType = ITypedElementDataStorageInterface::FQueryDescription::EActionType;
+	using CompletionType = ITypedElementDataStorageInterface::FQueryResult::ECompletion;
+
+	ITypedElementDataStorageInterface::FQueryResult Result;
+
+	if (FTypedElementExtendedQuery* QueryData = Get(Query))
+	{
+		if (QueryData->Description.bSimpleQuery)
+		{
+			switch (QueryData->Description.Action)
+			{
+			case ActionType::None:
+				Result.Completed = CompletionType::Fully;
+				break;
+			case ActionType::Select:
+				// Fall through: There's nothing to callback to, so only return the total count.
+			case ActionType::Count:
+				Result.Count = QueryData->NativeQuery.GetNumMatchingEntities(EntityManager);
+				Result.Completed = CompletionType::Fully;
+				break;
+			default:
+				Result.Completed = CompletionType::Unsupported;
+				break;
+			}
+		}
+		else
+		{
+			Result.Completed = CompletionType::Unsupported;
+		}
+	}
+	else
+	{
+		Result.Completed = CompletionType::Unavailable;
+	}
+
+	return Result;
+}
+
+ITypedElementDataStorageInterface::FQueryResult FTypedElementExtendedQueryStore::RunQuery(FMassEntityManager& EntityManager, Handle Query,
+	ITypedElementDataStorageInterface::DirectQueryCallbackRef Callback)
+{
+	using ActionType = ITypedElementDataStorageInterface::FQueryDescription::EActionType;
+	using CompletionType = ITypedElementDataStorageInterface::FQueryResult::ECompletion;
+
+	ITypedElementDataStorageInterface::FQueryResult Result;
+
+	if (FTypedElementExtendedQuery* QueryData = Get(Query))
+	{
+		if (QueryData->Description.bSimpleQuery)
+		{
+			switch (QueryData->Description.Action)
+			{
+			case ActionType::None:
+				Result.Completed = CompletionType::Fully;
+				break;
+			case ActionType::Select:
+				if (!QueryData->Processor.IsValid())
+				{
+					Result = FTypedElementQueryProcessorData::Execute(
+						Callback, QueryData->Description, QueryData->NativeQuery, EntityManager);
+				}
+				else
+				{
+					Result.Completed = CompletionType::Unsupported;
+				}
+				break;
+			case ActionType::Count:
+				// Only the count is requested so no need to trigger the callback.
+				Result.Count = QueryData->NativeQuery.GetNumMatchingEntities(EntityManager);
+				Result.Completed = CompletionType::Fully;
+				break;
+			default:
+				Result.Completed = CompletionType::Unsupported;
+				break;
+			}
+		}
+		else
+		{
+			Result.Completed = CompletionType::Unsupported;
+		}
+	}
+	else
+	{
+		Result.Completed = CompletionType::Unavailable;
+	}
+
+	return Result;
+}
+
+void FTypedElementExtendedQueryStore::RunPhasePreambleQueries(FMassEntityManager& EntityManager,
+	ITypedElementDataStorageInterface::EQueryTickPhase Phase, float DeltaTime)
+{
+	RunPhasePreOrPostAmbleQueries(EntityManager, Phase, DeltaTime, PhasePreparationQueries[static_cast<QueryTickPhaseType>(Phase)]);
+}
+
+void FTypedElementExtendedQueryStore::RunPhasePostambleQueries(FMassEntityManager& EntityManager,
+	ITypedElementDataStorageInterface::EQueryTickPhase Phase, float DeltaTime)
+{
+	RunPhasePreOrPostAmbleQueries(EntityManager, Phase, DeltaTime, PhaseFinalizationQueries[static_cast<QueryTickPhaseType>(Phase)]);
+}
+
+void FTypedElementExtendedQueryStore::DebugPrintQueryCallbacks(FOutputDevice& Output) const
+{
+	Output.Log(TEXT("The Typed Elements Data Storage has the following query callbacks:"));
+	Queries.ListAliveEntries(
+		[&Output](const FTypedElementExtendedQuery& Query)
+		{
+			if (Query.Processor)
+			{
+				Output.Logf(TEXT("    [%s] %s"),
+					IsValid(Query.Processor.Get()) ? TEXT("Valid") : TEXT("Invalid"),
+					*(Query.Processor->GetProcessorName()));
+			}
+		});
+
+	for (QueryTickPhaseType PhaseId = 0; PhaseId < static_cast<QueryTickPhaseType>(MaxTickPhase); ++PhaseId)
+	{
+		for (Handle QueryHandle : PhasePreparationQueries[PhaseId])
+		{
+			const FTypedElementExtendedQuery& QueryData = GetChecked(QueryHandle);
+			Output.Logf(TEXT("    [Valid] %s [Editor Phase Preamble]"), *QueryData.Description.Callback.Name.ToString());
+		}
+		for (Handle QueryHandle : PhaseFinalizationQueries[PhaseId])
+		{
+			const FTypedElementExtendedQuery& QueryData = GetChecked(QueryHandle);
+			Output.Logf(TEXT("    [Valid] %s [Editor Phase Postamble]"), *QueryData.Description.Callback.Name.ToString());
+		}
+	}
+
+	Output.Log(TEXT("End of Typed Elements Data Storage query callback list."));
+}
+
+FMassEntityQuery& FTypedElementExtendedQueryStore::SetupNativeQuery(
+	ITypedElementDataStorageInterface::FQueryDescription& Query, FTypedElementExtendedQuery& StoredQuery)
+{
+	using DSI = ITypedElementDataStorageInterface;
+
+	if (Query.Action == DSI::FQueryDescription::EActionType::Select)
+	{
+		switch (Query.Callback.Type)
+		{
+		case DSI::EQueryCallbackType::None:
+			break;
+		case DSI::EQueryCallbackType::Processor:
+		{
+			UTypedElementQueryProcessorCallbackAdapterProcessor* Processor =
+				NewObject<UTypedElementQueryProcessorCallbackAdapterProcessor>();
+			StoredQuery.Processor.Reset(Processor);
+			return Processor->GetQuery();
+		}
+		case DSI::EQueryCallbackType::ObserveAdd:
+			// Fall-through
+		case DSI::EQueryCallbackType::ObserveRemove:
+		{
+			UTypedElementQueryObserverCallbackAdapterProcessor* Observer =
+				NewObject<UTypedElementQueryObserverCallbackAdapterProcessor>();
+			StoredQuery.Processor.Reset(Observer);
+			return Observer->GetQuery();
+		}
+		case DSI::EQueryCallbackType::PhasePreparation:
+			break;
+		case DSI::EQueryCallbackType::PhaseFinalization:
+			break;
+		default:
+			checkf(false, TEXT("Unsupported query callback type %i."), static_cast<int>(Query.Callback.Type));
+			break;
+		}
+	}
+	return StoredQuery.NativeQuery;
+}
+
+bool FTypedElementExtendedQueryStore::SetupSelectedColumns(
+	ITypedElementDataStorageInterface::FQueryDescription& Query, FMassEntityQuery& NativeQuery)
+{
+	using DSI = ITypedElementDataStorageInterface;
+
+	switch (Query.Action)
+	{
+	case DSI::FQueryDescription::EActionType::None:
+		return true;
+	case DSI::FQueryDescription::EActionType::Select:
+	{
+		const int32 SelectionCount = Query.SelectionTypes.Num();
+		if(ensureMsgf(SelectionCount == Query.SelectionAccessTypes.Num(),
+			TEXT("The number of query selection types (%i) doesn't match the number of selection access types (%i)."),
+				SelectionCount, Query.SelectionAccessTypes.Num()))
+		{
+			for (int SelectionIndex = 0; SelectionIndex < SelectionCount; ++SelectionIndex)
+			{
+				TWeakObjectPtr<const UScriptStruct>& Type = Query.SelectionTypes[SelectionIndex];
+				
+				if (ensureMsgf(Type.IsValid(), TEXT("Provided query selection type can not be null.")) &&
+					ensureMsgf(
+						Type->IsChildOf(FTypedElementDataStorageColumn::StaticStruct()) ||
+						Type->IsChildOf(FMassFragment::StaticStruct()),
+						TEXT("Provided query selection type '%s' is not based on FTypedElementDataStorageColumn or another supported base type."),
+						*Type->GetStructPathName().ToString()))
+				{
+					NativeQuery.AddRequirement(Type.Get(), ConvertToNativeAccessType(Query.SelectionAccessTypes[SelectionIndex]));
+				}
+				else
+				{
+					return false;
+				}
+			}
+			return true;
+		}
+		return false;
+	}
+	case DSI::FQueryDescription::EActionType::Count:
+	{
+		bool bIsSelectionEmpty = Query.SelectionTypes.IsEmpty();
+		bool bIsAccessTypesEmpty = Query.SelectionAccessTypes.IsEmpty();
+		checkf(bIsSelectionEmpty, TEXT("Count queries for the Typed Elements Data Storage can't have entries for selection."));
+		checkf(bIsAccessTypesEmpty, TEXT("Count queries for the Typed Elements Data Storage can't have entries for selection."));
+		return bIsSelectionEmpty && bIsAccessTypesEmpty;
+	}
+	default:
+		checkf(false, TEXT("Unexpected query action: %i."), static_cast<int32>(Query.Action));
+		return false;
+	}
+}
+
+bool FTypedElementExtendedQueryStore::SetupConditions(
+	ITypedElementDataStorageInterface::FQueryDescription& Query, FMassEntityQuery& NativeQuery)
+{
+	using DSI = ITypedElementDataStorageInterface;
+
+	if (Query.ConditionTypes.IsEmpty())
+	{
+		return true;
+	}
+
+	if (ensureMsgf(Query.bSimpleQuery, TEXT("The Data Storage back-end currently only supports simple queries.")))
+	{
+		if (ensureMsgf(Query.ConditionTypes.Num() == Query.ConditionOperators.Num(),
+			TEXT("The types and operators for a typed element query have gone out of sync.")))
+		{
+			const DSI::FQueryDescription::FOperator* Operand = Query.ConditionOperators.GetData();
+			for (DSI::FQueryDescription::EOperatorType Type : Query.ConditionTypes)
+			{
+				EMassFragmentPresence Presence;
+				switch (Type)
+				{
+				case DSI::FQueryDescription::EOperatorType::SimpleAll:
+					Presence = EMassFragmentPresence::All;
+					break;
+				case DSI::FQueryDescription::EOperatorType::SimpleAny:
+					Presence = EMassFragmentPresence::Any;
+					break;
+				case DSI::FQueryDescription::EOperatorType::SimpleNone:
+					Presence = EMassFragmentPresence::None;
+					break;
+				default:
+					continue;
+				}
+
+				if (Operand->Type->IsChildOf(FMassTag::StaticStruct()))
+				{
+					NativeQuery.AddTagRequirement(*(Operand->Type), Presence);
+				}
+				else if (Operand->Type->IsChildOf(FMassFragment::StaticStruct()))
+				{
+					NativeQuery.AddRequirement(Operand->Type.Get(), EMassFragmentAccess::None, Presence);
+				}
+
+				++Operand;
+			}
+			return true;
+		}
+	}
+	return false;
+}
+
+bool FTypedElementExtendedQueryStore::SetupDependencies(
+	ITypedElementDataStorageInterface::FQueryDescription& Query, FMassEntityQuery& NativeQuery)
+{
+	using DSI = ITypedElementDataStorageInterface;
+
+	const int32 DependencyCount = Query.DependencyTypes.Num();
+	if (ensureMsgf(DependencyCount == Query.DependencyFlags.Num() && DependencyCount == Query.CachedDependencies.Num(),
+		TEXT("The number of query dependencies (%i) doesn't match the number of dependency access types (%i) and/or cached dependencies count (%i)."),
+		DependencyCount, Query.DependencyFlags.Num(), Query.CachedDependencies.Num()))
+	{
+		for (int32 DependencyIndex = 0; DependencyIndex < DependencyCount; ++DependencyIndex)
+		{
+			TWeakObjectPtr<const UClass>& Type = Query.DependencyTypes[DependencyIndex];
+			if (ensureMsgf(Type.IsValid(), TEXT("Provided query dependency type can not be null.")) &&
+				ensureMsgf(Type->IsChildOf<USubsystem>(), TEXT("Provided query dependency type '%s' is not based on USubSystem."),
+					*Type->GetStructPathName().ToString()))
+			{
+				DSI::EQueryDependencyFlags Flags = Query.DependencyFlags[DependencyIndex];
+				NativeQuery.AddSubsystemRequirement(
+					const_cast<UClass*>(Type.Get()),
+					EnumHasAllFlags(Flags, DSI::EQueryDependencyFlags::ReadOnly) ? EMassFragmentAccess::ReadOnly : EMassFragmentAccess::ReadWrite,
+					EnumHasAllFlags(Flags, DSI::EQueryDependencyFlags::GameThreadBound));
+			}
+			else
+			{
+				return false;
+			}
+		}
+		return true;
+	}
+	return false;
+}
+
+bool FTypedElementExtendedQueryStore::SetupTickGroupDefaults(ITypedElementDataStorageInterface::FQueryDescription& Query)
+{
+	const FTickGroupDescription* TickGroup = TickGroupDescriptions.Find({ Query.Callback.Group, Query.Callback.Phase });
+	if (TickGroup)
+	{
+		for (auto It = Query.Callback.BeforeGroups.CreateIterator(); It; ++It)
+		{
+			if (TickGroup->BeforeGroups.Contains(*It))
+			{
+				It.RemoveCurrentSwap();
+			}
+		}
+		Query.Callback.BeforeGroups.Append(TickGroup->BeforeGroups);
+		
+		for (auto It = Query.Callback.AfterGroups.CreateIterator(); It; ++It)
+		{
+			if (TickGroup->AfterGroups.Contains(*It))
+			{
+				It.RemoveCurrentSwap();
+			}
+		}
+		Query.Callback.AfterGroups.Append(TickGroup->AfterGroups);
+		
+		Query.Callback.bForceToGameThread |= TickGroup->bRequiresMainThread;
+	}
+	return true;
+}
+
+bool FTypedElementExtendedQueryStore::SetupProcessors(Handle Query, FTypedElementExtendedQuery& StoredQuery, 
+	FMassEntityManager& EntityManager, FMassProcessingPhaseManager& PhaseManager)
+{
+	using DSI = ITypedElementDataStorageInterface;
+
+	// Register Phase processors locally.
+	switch (StoredQuery.Description.Callback.Type)
+	{
+	case DSI::EQueryCallbackType::PhasePreparation:
+		RegisterPreambleQuery(StoredQuery.Description.Callback.Phase, Query);
+		break;
+	case DSI::EQueryCallbackType::PhaseFinalization:
+		RegisterPostambleQuery(StoredQuery.Description.Callback.Phase, Query);
+		break;
+	}
+
+	// Register regular processors and observer with Mass.
+	if (StoredQuery.Processor)
+	{
+		if (StoredQuery.Processor->IsA<UTypedElementQueryProcessorCallbackAdapterProcessor>())
+		{
+			static_cast<UTypedElementQueryProcessorCallbackAdapterProcessor*>(StoredQuery.Processor.Get())->
+				ConfigureQueryCallback(StoredQuery);
+			PhaseManager.RegisterDynamicProcessor(*StoredQuery.Processor);
+		}
+		else if (StoredQuery.Processor->IsA<UTypedElementQueryObserverCallbackAdapterProcessor>())
+		{
+			UTypedElementQueryObserverCallbackAdapterProcessor* Observer =
+				static_cast<UTypedElementQueryObserverCallbackAdapterProcessor*>(StoredQuery.Processor.Get());
+			Observer->ConfigureQueryCallback(StoredQuery);
+			EntityManager.GetObserverManager().AddObserverInstance(*Observer->GetObservedType(), Observer->GetObservedOperation(), *Observer);
+		}
+		else
+		{
+			checkf(false, TEXT("Query processor %s is of unsupported type %s."),
+				*StoredQuery.Description.Callback.Name.ToString(), *StoredQuery.Processor->GetSparseClassDataStruct()->GetName());
+			return false;
+		}
+	}
+	return true;
+}
+
+EMassFragmentAccess FTypedElementExtendedQueryStore::ConvertToNativeAccessType(ITypedElementDataStorageInterface::EQueryAccessType AccessType)
+{
+	switch (AccessType)
+	{
+	case ITypedElementDataStorageInterface::EQueryAccessType::ReadOnly:
+		return EMassFragmentAccess::ReadOnly;
+	case ITypedElementDataStorageInterface::EQueryAccessType::ReadWrite:
+		return EMassFragmentAccess::ReadWrite;
+	default:
+		checkf(false, TEXT("Invalid query access type: %i."), static_cast<uint32>(AccessType));
+		return EMassFragmentAccess::MAX;
+	}
+}
+
+void FTypedElementExtendedQueryStore::RegisterPreambleQuery(ITypedElementDataStorageInterface::EQueryTickPhase Phase, Handle Query)
+{
+	PhasePreparationQueries[static_cast<QueryTickPhaseType>(Phase)].Add(Query);
+}
+
+void FTypedElementExtendedQueryStore::RegisterPostambleQuery(ITypedElementDataStorageInterface::EQueryTickPhase Phase, Handle Query)
+{
+	PhaseFinalizationQueries[static_cast<QueryTickPhaseType>(Phase)].Add(Query);
+}
+
+void FTypedElementExtendedQueryStore::UnregisterPreambleQuery(ITypedElementDataStorageInterface::EQueryTickPhase Phase, Handle Query)
+{
+	int32 Index;
+	if (PhasePreparationQueries[static_cast<QueryTickPhaseType>(Phase)].Find(Query, Index))
+	{
+		PhasePreparationQueries[static_cast<QueryTickPhaseType>(Phase)].RemoveAt(Index);
+	}
+}
+
+void FTypedElementExtendedQueryStore::UnregisterPostambleQuery(ITypedElementDataStorageInterface::EQueryTickPhase Phase, Handle Query)
+{
+	int32 Index;
+	if (PhaseFinalizationQueries[static_cast<QueryTickPhaseType>(Phase)].Find(Query, Index))
+	{
+		PhaseFinalizationQueries[static_cast<QueryTickPhaseType>(Phase)].RemoveAt(Index);
+	}
+}
+
+void FTypedElementExtendedQueryStore::RunPhasePreOrPostAmbleQueries(FMassEntityManager& EntityManager,
+	ITypedElementDataStorageInterface::EQueryTickPhase Phase, float DeltaTime, TArray<Handle>& QueryHandles)
+{
+	if (!QueryHandles.IsEmpty())
+	{
+		FPhasePreOrPostAmbleExecutor Executor(EntityManager, DeltaTime);
+		for (Handle Query : QueryHandles)
+		{
+			FTypedElementExtendedQuery& QueryData = Queries.Get(Query);
+			Executor.ExecuteQuery(QueryData.Description, QueryData.NativeQuery, QueryData.Description.Callback.Function);
+		}
+	}
+}
diff --git a/Engine/Plugins/Experimental/TypedElementsDataStorage/Source/TypedElementsDataStorage/Private/Queries/TypedElementExtendedQueryStore.h b/Engine/Plugins/Experimental/TypedElementsDataStorage/Source/TypedElementsDataStorage/Private/Queries/TypedElementExtendedQueryStore.h
new file mode 100644
--- /dev/null
+++ b/Engine/Plugins/Experimental/TypedElementsDataStorage/Source/TypedElementsDataStorage/Private/Queries/TypedElementExtendedQueryStore.h
@@ -0,0 +1,135 @@
+// Copyright Epic Games, Inc. All Rights Reserved.
+
+#pragma once
+
+#include "Elements/Interfaces/TypedElementDataStorageInterface.h"
+#include "MassEntityQuery.h"
+#include "MassRequirements.h"
+#include "TypedElementHandleStore.h"
+#include "UObject/StrongObjectPtr.h"
+
+struct FMassEntityManager;
+struct FMassProcessingPhaseManager;
+class FOutputDevice;
+class UMassProcessor;
+
+struct FTypedElementExtendedQuery
+{
+	FMassEntityQuery NativeQuery; // Used if there's no processor bound.
+	ITypedElementDataStorageInterface::FQueryDescription Description;
+	TStrongObjectPtr<UMassProcessor> Processor;
+};
+
+/**
+ * Storage and utilities for Typed Element queries after they've been processed by the Data Storage implementation.
+ */
+class FTypedElementExtendedQueryStore
+{
+private:
+	using QueryStore = TTypedElementHandleStore<FTypedElementExtendedQuery>;
+public:
+	using Handle = QueryStore::Handle;
+	using ListAliveEntriesCallback = QueryStore::ListAliveEntriesCallback;
+
+	/**
+	 * @section Registration
+	 * @description A set of functions to manage the registration of queries.
+	 */
+
+	/** Adds a new query to the store and initializes the query with the provided arguments. */
+	Handle RegisterQuery(ITypedElementDataStorageInterface::FQueryDescription Query, FMassEntityManager& EntityManager, 
+		FMassProcessingPhaseManager& PhaseManager);
+	/** Removes the query at the given handle if still alive and otherwise does nothing. */
+	void UnregisterQuery(Handle Query, FMassProcessingPhaseManager& PhaseManager);
+	
+	/** Register the defaults for a tick group. These will be applied on top of any settings provided with a query registration. */
+	void RegisterTickGroup(FName GroupName, ITypedElementDataStorageInterface::EQueryTickPhase Phase,
+		FName BeforeGroup, FName AfterGroup, bool bRequiresMainThread);
+	/** Removes a previously registered set of tick group defaults. */
+	void UnregisterTickGroup(FName GroupName, ITypedElementDataStorageInterface::EQueryTickPhase Phase);
+
+	/**
+	 * @section Retrieval
+	 * @description Functions to retrieve data or information on queries.
+	 */
+
+	/** Retrieves the query at the provided handle, if still alive or otherwise returns a nullptr. */
+	FTypedElementExtendedQuery* Get(Handle Entry);
+	/** Retrieves the query at the provided handle, if still alive or otherwise returns a nullptr. */
+	FTypedElementExtendedQuery* GetMutable(Handle Entry);
+	/** Retrieves the query at the provided handle, if still alive or otherwise returns a nullptr. */
+	const FTypedElementExtendedQuery* Get(Handle Entry) const;
+
+	/** Retrieves the query at the provided handle, if still alive. It's up to the caller to guarantee the query is still alive. */
+	FTypedElementExtendedQuery& GetChecked(Handle Entry);
+	/** Retrieves the query at the provided handle, if still alive. It's up to the caller to guarantee the query is still alive. */
+	FTypedElementExtendedQuery& GetMutableChecked(Handle Entry);
+	/** Retrieves the query at the provided handle, if still alive. It's up to the caller to guarantee the query is still alive. */
+	const FTypedElementExtendedQuery& GetChecked(Handle Entry) const;
+
+	/** Gets the original description used to create an extended query or an empty default if the provided query isn't alive. */
+	const ITypedElementDataStorageInterface::FQueryDescription& GetQueryDescription(Handle Query) const;
+
+	/** Checks to see if a query is still available or has been removed. */
+	bool IsAlive(Handle Entry) const;
+
+	/** Calls the provided callback for each query that's available. */
+	void ListAliveEntries(const ListAliveEntriesCallback& Callback) const;
+
+	/**
+	 * @section Execution
+	 * @description Various functions to run queries.
+	 */
+	ITypedElementDataStorageInterface::FQueryResult RunQuery(FMassEntityManager& EntityManager, Handle Query);
+	ITypedElementDataStorageInterface::FQueryResult RunQuery(FMassEntityManager& EntityManager, Handle Query,
+		ITypedElementDataStorageInterface::DirectQueryCallbackRef Callback);
+	void RunPhasePreambleQueries(FMassEntityManager& EntityManager, ITypedElementDataStorageInterface::EQueryTickPhase Phase, float DeltaTime);
+	void RunPhasePostambleQueries(FMassEntityManager& EntityManager, ITypedElementDataStorageInterface::EQueryTickPhase Phase, float DeltaTime);
+
+	void DebugPrintQueryCallbacks(FOutputDevice& Output) const;
+
+private:
+	using QueryTickPhaseType = std::underlying_type_t<ITypedElementDataStorageInterface::EQueryTickPhase>;
+	static constexpr QueryTickPhaseType MaxTickPhase = static_cast<QueryTickPhaseType>(ITypedElementDataStorageInterface::EQueryTickPhase::Max);
+
+	struct FTickGroupId
+	{
+		FName Name;
+		ITypedElementDataStorageInterface::EQueryTickPhase Phase;
+
+		friend inline uint32 GetTypeHash(const FTickGroupId& Id) { return HashCombine(GetTypeHash(Id.Name), GetTypeHash(Id.Phase)); }
+		friend inline bool operator==(const FTickGroupId& Lhs, const FTickGroupId& Rhs) { return Lhs.Phase == Rhs.Phase && Lhs.Name == Rhs.Name; }
+		friend inline bool operator!=(const FTickGroupId& Lhs, const FTickGroupId& Rhs) { return Lhs.Phase != Rhs.Phase || Lhs.Name != Rhs.Name; }
+	};
+
+	struct FTickGroupDescription
+	{
+		TArray<FName> BeforeGroups;
+		TArray<FName> AfterGroups;
+		bool bRequiresMainThread{ false };
+	};
+
+	FMassEntityQuery& SetupNativeQuery(ITypedElementDataStorageInterface::FQueryDescription& Query, FTypedElementExtendedQuery& StoredQuery);
+	bool SetupSelectedColumns(ITypedElementDataStorageInterface::FQueryDescription& Query, FMassEntityQuery& NativeQuery);
+	bool SetupConditions(ITypedElementDataStorageInterface::FQueryDescription& Query, FMassEntityQuery& NativeQuery);
+	bool SetupDependencies(ITypedElementDataStorageInterface::FQueryDescription& Query, FMassEntityQuery& NativeQuery);
+	bool SetupTickGroupDefaults(ITypedElementDataStorageInterface::FQueryDescription& Query);
+	bool SetupProcessors(Handle Query, FTypedElementExtendedQuery& StoredQuery,
+		FMassEntityManager& EntityManager, FMassProcessingPhaseManager& PhaseManager);
+
+	EMassFragmentAccess ConvertToNativeAccessType(ITypedElementDataStorageInterface::EQueryAccessType AccessType);
+
+	void RegisterPreambleQuery(ITypedElementDataStorageInterface::EQueryTickPhase Phase, Handle Query);
+	void RegisterPostambleQuery(ITypedElementDataStorageInterface::EQueryTickPhase Phase, Handle Query);
+	void UnregisterPreambleQuery(ITypedElementDataStorageInterface::EQueryTickPhase Phase, Handle Query);
+	void UnregisterPostambleQuery(ITypedElementDataStorageInterface::EQueryTickPhase Phase, Handle Query);
+	void RunPhasePreOrPostAmbleQueries(FMassEntityManager& EntityManager,
+		ITypedElementDataStorageInterface::EQueryTickPhase Phase, float DeltaTime, TArray<Handle>& QueryHandles);
+
+	static const ITypedElementDataStorageInterface::FQueryDescription EmptyDescription;
+
+	QueryStore Queries;
+	TMap<FTickGroupId, FTickGroupDescription> TickGroupDescriptions;
+	TArray<Handle> PhasePreparationQueries[MaxTickPhase];
+	TArray<Handle> PhaseFinalizationQueries[MaxTickPhase];
+};
diff --git a/Engine/Plugins/Experimental/TypedElementsDataStorage/Source/TypedElementsDataStorage/Private/TypedElementDatabase.cpp b/Engine/Plugins/Experimental/TypedElementsDataStorage/Source/TypedElementsDataStorage/Private/TypedElementDatabase.cpp
--- a/Engine/Plugins/Experimental/TypedElementsDataStorage/Source/TypedElementsDataStorage/Private/TypedElementDatabase.cpp
+++ b/Engine/Plugins/Experimental/TypedElementsDataStorage/Source/TypedElementsDataStorage/Private/TypedElementDatabase.cpp
@@ -8,6 +8,7 @@
 #include "MassCommonTypes.h"
 #include "MassEntityEditorSubsystem.h"
 #include "MassEntityTypes.h"
+#include "MassProcessor.h"
 #include "MassSimulationSubsystem.h"
 #include "MassSubsystemAccess.h"
 #include "Processors/TypedElementProcessorAdaptors.h"
@@ -540,309 +541,36 @@ bool UTypedElementDatabase::HasColumns(TypedElementRowHandle Row, TConstArrayVie
 void UTypedElementDatabase::RegisterTickGroup(
 	FName GroupName, EQueryTickPhase Phase, FName BeforeGroup, FName AfterGroup, bool bRequiresMainThread)
 {
-	FTickGroupDescription& Group = TickGroupDescriptions.FindOrAdd({ GroupName, Phase });
-
-	if (!Group.BeforeGroups.Find(BeforeGroup))
-	{
-		Group.BeforeGroups.Add(BeforeGroup);
-	}
-
-	if (!Group.AfterGroups.Find(AfterGroup))
-	{
-		Group.AfterGroups.Add(AfterGroup);
-	}
-
-	if (bRequiresMainThread)
-	{
-		Group.bRequiresMainThread = true;
-	}
+	Queries.RegisterTickGroup(GroupName, Phase, BeforeGroup, AfterGroup, bRequiresMainThread);
 }
 
 void UTypedElementDatabase::UnregisterTickGroup(FName GroupName, EQueryTickPhase Phase)
 {
-	TickGroupDescriptions.Remove({ GroupName, Phase });
+	Queries.UnregisterTickGroup(GroupName, Phase);
 }
 
 TypedElementQueryHandle UTypedElementDatabase::RegisterQuery(FQueryDescription&& Query)
 {
-	auto LocalToNativeAccess = [](EQueryAccessType Access) -> EMassFragmentAccess
-	{
-		switch (Access)
-		{
-			case EQueryAccessType::ReadOnly:
-				return EMassFragmentAccess::ReadOnly;
-			case EQueryAccessType::ReadWrite:
-				return EMassFragmentAccess::ReadWrite;
-			default:
-				checkf(false, TEXT("Invalid query access type: %i."), static_cast<uint32>(Access));
-				return EMassFragmentAccess::MAX;
-		}
-	};
-
-	auto SetupNativeQuery = [](FQueryDescription& Query, FTypedElementDatabaseExtendedQuery& StoredQuery) -> FMassEntityQuery&
-	{
-		if (Query.Action == FQueryDescription::EActionType::Select)
-		{
-			switch (Query.Callback.Type)
-			{
-				case EQueryCallbackType::None:
-					break;
-				case EQueryCallbackType::Processor:
-				{
-					UTypedElementQueryProcessorCallbackAdapterProcessor* Processor = 
-						NewObject<UTypedElementQueryProcessorCallbackAdapterProcessor>();
-					StoredQuery.Processor.Reset(Processor);
-					return Processor->GetQuery();
-				}
-				case EQueryCallbackType::ObserveAdd:
-					// Fallthrough
-				case EQueryCallbackType::ObserveRemove:
-				{
-					UTypedElementQueryObserverCallbackAdapterProcessor* Observer = 
-						NewObject<UTypedElementQueryObserverCallbackAdapterProcessor>();
-					StoredQuery.Processor.Reset(Observer);
-					return Observer->GetQuery();
-				}
-				case EQueryCallbackType::PhasePreparation:
-					break;
-				case EQueryCallbackType::PhaseFinalization:
-					break;
-				default:
-					checkf(false, TEXT("Unsupported query callback type %i."), static_cast<int>(Query.Callback.Type));
-					break;
-			}
-		}
-		return StoredQuery.NativeQuery;
-	};
-
-	QueryStore::Handle Result = Queries.Emplace();
-	FTypedElementDatabaseExtendedQuery& StoredQuery = Queries.GetMutable(Result);
-	
-	FMassEntityQuery& NativeQuery = SetupNativeQuery(Query, StoredQuery);
-
-	// Setup selected columns section
-	if (Query.Action == FQueryDescription::EActionType::Count)
-	{
-		checkf(Query.SelectionTypes.IsEmpty(), TEXT("Count queries for the Typed Elements Data Storage can't have entries for selection."));
-		checkf(Query.SelectionAccessTypes.IsEmpty(), TEXT("Count queries for the Typed Elements Data Storage can't have entries for selection."));
-	}
-	else if (Query.Action == FQueryDescription::EActionType::Select)
-	{
-		const int32 SelectionCount = Query.SelectionTypes.Num();
-		checkf(SelectionCount == Query.SelectionAccessTypes.Num(),
-			TEXT("The number of query selection types (%i) doesn't match the number of selection access types (%i)."),
-			SelectionCount, Query.SelectionAccessTypes.Num());
-		for (int SelectionIndex = 0; SelectionIndex < SelectionCount; ++SelectionIndex)
-		{
-			TWeakObjectPtr<const UScriptStruct>& Type = Query.SelectionTypes[SelectionIndex];
-			checkf(Type.IsValid(), TEXT("Provided query selection type can not be null."));
-			checkf(
-				Type->IsChildOf(FTypedElementDataStorageColumn::StaticStruct()) ||
-				Type->IsChildOf(FMassFragment::StaticStruct()) || 
-				Type->IsChildOf(FMassTag::StaticStruct()),
-				TEXT("Provided query selection type '%s' is not based on FTypedElementDataStorageColumn or another supported base type."), 
-				*Type->GetStructPathName().ToString());
-			NativeQuery.AddRequirement(Type.Get(), LocalToNativeAccess(Query.SelectionAccessTypes[SelectionIndex]));
-		}
-	}
-	else
-	{
-		checkf(Query.Action == FQueryDescription::EActionType::None, TEXT("Unexpected query action: %i."), static_cast<int32>(Query.Action));
-	}
-
-	// Configure conditions
-	if (Query.bSimpleQuery) // This backend currently only supports simple queries.
-	{
-		checkf(Query.ConditionTypes.Num() == Query.ConditionOperators.Num(),
-			TEXT("The types and operators for a typed element query have gone out of sync."));
-		
-		const FQueryDescription::FOperator* Operand = Query.ConditionOperators.GetData();
-		for (FQueryDescription::EOperatorType Type : Query.ConditionTypes)
-		{
-			EMassFragmentPresence Presence;
-			switch (Type)
-			{
-				case FQueryDescription::EOperatorType::SimpleAll:
-					Presence = EMassFragmentPresence::All;
-					break;
-				case FQueryDescription::EOperatorType::SimpleAny:
-					Presence = EMassFragmentPresence::Any;
-					break;
-				case FQueryDescription::EOperatorType::SimpleNone:
-					Presence = EMassFragmentPresence::None;
-					break;
-				default:
-					continue;
-			}
-
-			if (Operand->Type->IsChildOf(FMassTag::StaticStruct()))
-			{
-				NativeQuery.AddTagRequirement(*(Operand->Type), Presence);
-			}
-			else if (Operand->Type->IsChildOf(FMassFragment::StaticStruct()))
-			{
-				NativeQuery.AddRequirement(Operand->Type.Get(), EMassFragmentAccess::None, Presence);
-			}
-
-			++Operand;
-		}
-	}
-
-	// Assign dependencies.
-	const int32 DependencyCount = Query.DependencyTypes.Num();
-	checkf(DependencyCount == Query.DependencyFlags.Num() && DependencyCount == Query.CachedDependencies.Num(),
-		TEXT("The number of query depedencies (%i) doesn't match the number of dependency access types (%i) and/or cached dependencies count (%i)."),
-		DependencyCount, Query.DependencyFlags.Num(), Query.CachedDependencies.Num());
-	for (int32 DependencyIndex = 0; DependencyIndex < DependencyCount; ++DependencyIndex)
-	{
-		TWeakObjectPtr<const UClass>& Type = Query.DependencyTypes[DependencyIndex];
-		checkf(Type.IsValid(), TEXT("Provided query dependcy type can not be null."));
-		checkf(Type->IsChildOf<USubsystem>(), TEXT("Provided query dependency type '%s' is not based on USubSystem."), 
-			*Type->GetStructPathName().ToString());
-		
-		EQueryDependencyFlags Flags = Query.DependencyFlags[DependencyIndex];
-		NativeQuery.AddSubsystemRequirement(
-			const_cast<UClass*>(Type.Get()), 
-			EnumHasAllFlags(Flags, EQueryDependencyFlags::ReadOnly) ? EMassFragmentAccess::ReadOnly : EMassFragmentAccess::ReadWrite,
-			EnumHasAllFlags(Flags, EQueryDependencyFlags::GameThreadBound));
-	}
-
-	// Copy pre-registered phase and group information.
-	const FTickGroupDescription* TickGroup = TickGroupDescriptions.Find({ Query.Callback.Group, Query.Callback.Phase });
-	if (TickGroup)
-	{
-		for (auto It = Query.Callback.BeforeGroups.CreateIterator(); It; ++It)
-		{
-			if (TickGroup->BeforeGroups.Contains(*It))
-			{
-				It.RemoveCurrentSwap();
-			}
-		}
-		Query.Callback.BeforeGroups.Append(TickGroup->BeforeGroups);
-		for (auto It = Query.Callback.AfterGroups.CreateIterator(); It; ++It)
-		{
-			if (TickGroup->AfterGroups.Contains(*It))
-			{
-				It.RemoveCurrentSwap();
-			}
-		}
-		Query.Callback.AfterGroups.Append(TickGroup->AfterGroups);
-		if (TickGroup->bRequiresMainThread)
-		{
-			Query.Callback.bForceToGameThread = true;
-		}
-	}
-
-	// Register Phase processors locally.
-	switch (Query.Callback.Type)
-	{
-	case EQueryCallbackType::PhasePreparation:
-		PhasePreparationQueries[static_cast<std::underlying_type_t<EQueryTickPhase>>(Query.Callback.Phase)].Add(Result.Handle);
-		break;
-	case EQueryCallbackType::PhaseFinalization:
-		PhaseFinalizationQueries[static_cast<std::underlying_type_t<EQueryTickPhase>>(Query.Callback.Phase)].Add(Result.Handle);
-		break;
-	}
-
-	StoredQuery.Description = MoveTemp(Query);
-
-	// Register regular processors and observer with Mass.
-	if (StoredQuery.Processor)
-	{
-		if (StoredQuery.Processor->IsA<UTypedElementQueryProcessorCallbackAdapterProcessor>())
-		{
-			if (UMassEntityEditorSubsystem* Mass = GEditor->GetEditorSubsystem<UMassEntityEditorSubsystem>())
-			{
-				static_cast<UTypedElementQueryProcessorCallbackAdapterProcessor*>(StoredQuery.Processor.Get())->
-					ConfigureQueryCallback(StoredQuery);
-				Mass->RegisterDynamicProcessor(*StoredQuery.Processor);
-			}
-		}
-		else if (StoredQuery.Processor->IsA<UTypedElementQueryObserverCallbackAdapterProcessor>())
-		{
-			UTypedElementQueryObserverCallbackAdapterProcessor* Observer =
-				static_cast<UTypedElementQueryObserverCallbackAdapterProcessor*>(StoredQuery.Processor.Get());
-			Observer->ConfigureQueryCallback(StoredQuery);
-			ActiveEditorEntityManager->GetObserverManager().AddObserverInstance(
-				*Observer->GetObservedType(), Observer->GetObservedOperation(), *Observer);
-		}
-		else
-		{
-			checkf(false, TEXT("Query processor %s is of unsupported type %s."), 
-				*Query.Callback.Name.ToString(), *StoredQuery.Processor->GetSparseClassDataStruct()->GetName());
-		}
-	}
-
-	return Result.Handle;
+	return (ActiveEditorEntityManager && ActiveEditorPhaseManager)
+		? Queries.RegisterQuery(MoveTemp(Query), *ActiveEditorEntityManager, *ActiveEditorPhaseManager).Handle
+		: TypedElementInvalidQueryHandle;
 }
 
 void UTypedElementDatabase::UnregisterQuery(TypedElementQueryHandle Query)
 {
-	QueryStore::Handle Handle;
-	Handle.Handle = Query;
-
-	if (Queries.IsAlive(Handle))
+	if (ActiveEditorPhaseManager)
 	{
-		FTypedElementDatabaseExtendedQuery& QueryData = Queries.Get(Handle);
-		if (QueryData.Processor)
-		{
-			if (QueryData.Processor->IsA<UTypedElementQueryProcessorCallbackAdapterProcessor>())
-			{
-				if (UMassEntityEditorSubsystem* Mass = GEditor->GetEditorSubsystem<UMassEntityEditorSubsystem>())
-				{
-					Mass->UnregisterDynamicProcessor(*QueryData.Processor);
-				}
-			}
-			else if (QueryData.Processor->IsA<UTypedElementQueryObserverCallbackAdapterProcessor>())
-			{
-				checkf(false, TEXT("Observer queries can not be unregistered."));
-			}
-			else
-			{
-				checkf(false, TEXT("Query processor %s is of unsupported type %s."),
-					*QueryData.Description.Callback.Name.ToString(), *QueryData.Processor->GetSparseClassDataStruct()->GetName());
-			}
-		}
-		else if (QueryData.Description.Callback.Type == EQueryCallbackType::PhasePreparation)
-		{
-			int32 Index;
-			if (PhasePreparationQueries[static_cast<std::underlying_type_t<EQueryTickPhase>>(QueryData.Description.Callback.Phase)].Find(Query, Index))
-			{
-				PhasePreparationQueries[static_cast<std::underlying_type_t<EQueryTickPhase>>(QueryData.Description.Callback.Phase)].RemoveAt(Index);
-			}
-		}
-		else if (QueryData.Description.Callback.Type == EQueryCallbackType::PhaseFinalization)
-		{
-			int32 Index;
-			if (PhaseFinalizationQueries[static_cast<std::underlying_type_t<EQueryTickPhase>>(QueryData.Description.Callback.Phase)].Find(Query, Index))
-			{
-				PhaseFinalizationQueries[static_cast<std::underlying_type_t<EQueryTickPhase>>(QueryData.Description.Callback.Phase)].RemoveAt(Index);
-			}
-		}
-		else
-		{
-			QueryData.NativeQuery.Clear();
-		}
+		FTypedElementExtendedQueryStore::Handle Handle;
+		Handle.Handle = Query;
+		Queries.UnregisterQuery(Handle, *ActiveEditorPhaseManager);
 	}
-
-	Queries.Remove(Handle);
 }
 
 const ITypedElementDataStorageInterface::FQueryDescription& UTypedElementDatabase::GetQueryDescription(TypedElementQueryHandle Query) const
 {
-	QueryStore::Handle Handle;
+	FTypedElementExtendedQueryStore::Handle Handle;
 	Handle.Handle = Query;
-
-	if (Queries.IsAlive(Handle))
-	{
-		const FTypedElementDatabaseExtendedQuery& QueryData = Queries.Get(Handle);
-		return QueryData.Description;
-	}
-	else
-	{
-		static ITypedElementDataStorageInterface::FQueryDescription EmptyDescription;
-		return EmptyDescription;
-	}
+	return Queries.GetQueryDescription(Handle);
 }
 
 FName UTypedElementDatabase::GetQueryTickGroupName(EQueryTickGroups Group) const
@@ -866,52 +594,17 @@ FName UTypedElementDatabase::GetQueryTickGroupName(EQueryTickGroups Group) const
 ITypedElementDataStorageInterface::FQueryResult UTypedElementDatabase::RunQuery(TypedElementQueryHandle Query)
 {
 	TRACE_CPUPROFILER_EVENT_SCOPE(TEDS.RunQuery);
-	
-	FQueryResult Result;
-
-	QueryStore::Handle Handle;
-	Handle.Handle = Query;
 
-	if (Queries.IsAlive(Handle))
+	if (ActiveEditorEntityManager)
 	{
-		FTypedElementDatabaseExtendedQuery& QueryData = Queries.Get(Handle);
-		if (QueryData.Description.bSimpleQuery)
-		{
-			switch (QueryData.Description.Action)
-			{
-				case FQueryDescription::EActionType::None:
-					Result.Completed = FQueryResult::ECompletion::Fully;
-					break;
-				case FQueryDescription::EActionType::Select:
-					// Fallthrough: There's nothing to callback to, so only return the total count.
-				case FQueryDescription::EActionType::Count:
-					if (ActiveEditorEntityManager)
-					{
-						Result.Count = QueryData.NativeQuery.GetNumMatchingEntities(*ActiveEditorEntityManager);
-						Result.Completed = FQueryResult::ECompletion::Fully;
-					}
-					else
-					{
-						Result.Completed = FQueryResult::ECompletion::Unavailable;
-					}
-					break;
-				default:
-					Result.Completed = FQueryResult::ECompletion::Unsupported;
-					break;
-			}
-		}
-		else
-		{
-			checkf(false, TEXT("Support for this option will be coming in a future update."));
-			Result.Completed = FQueryResult::ECompletion::Unsupported;
-		}
+		FTypedElementExtendedQueryStore::Handle Handle;
+		Handle.Handle = Query;
+		return Queries.RunQuery(*ActiveEditorEntityManager, Handle);
 	}
 	else
 	{
-		Result.Completed = FQueryResult::ECompletion::Unavailable;
+		return FQueryResult();
 	}
-
-	return Result;
 }
 
 ITypedElementDataStorageInterface::FQueryResult UTypedElementDatabase::RunQuery(
@@ -919,68 +612,16 @@ ITypedElementDataStorageInterface::FQueryResult UTypedElementDatabase::RunQuery(
 {
 	TRACE_CPUPROFILER_EVENT_SCOPE(TEDS.RunQuery);
 
-	FQueryResult Result;
-
-	QueryStore::Handle Handle;
-	Handle.Handle = Query;
-
-	if (Queries.IsAlive(Handle))
+	if (ActiveEditorEntityManager)
 	{
-		FTypedElementDatabaseExtendedQuery& QueryData = Queries.Get(Handle);
-		if (QueryData.Description.bSimpleQuery)
-		{
-			switch (QueryData.Description.Action)
-			{
-			case FQueryDescription::EActionType::None:
-				Result.Completed = FQueryResult::ECompletion::Fully;
-				break;
-			case FQueryDescription::EActionType::Select:
-				if (ActiveEditorEntityManager)
-				{
-					if (!QueryData.Processor.IsValid())
-					{
-						Result = FTypedElementQueryProcessorData::Execute(
-							Callback, QueryData.Description, QueryData.NativeQuery, *ActiveEditorEntityManager);
-					}
-					else
-					{
-						Result.Completed = FQueryResult::ECompletion::Unsupported;
-					}
-				}
-				else
-				{
-					Result.Completed = FQueryResult::ECompletion::Unavailable;
-				}
-				break;
-			case FQueryDescription::EActionType::Count:
-				// Only the count is requested so no need to trigger the callback.
-				if (ActiveEditorEntityManager)
-				{
-					Result.Count = QueryData.NativeQuery.GetNumMatchingEntities(*ActiveEditorEntityManager);
-					Result.Completed = FQueryResult::ECompletion::Fully;
-				}
-				else
-				{
-					Result.Completed = FQueryResult::ECompletion::Unavailable;
-				}
-				break;
-			default:
-				Result.Completed = FQueryResult::ECompletion::Unsupported;
-				break;
-			}
-		}
-		else
-		{
-			checkf(false, TEXT("Support for this option will be coming in a future update."));
-			Result.Completed = FQueryResult::ECompletion::Unsupported;
-		}
+		FTypedElementExtendedQueryStore::Handle Handle;
+		Handle.Handle = Query;
+		return Queries.RunQuery(*ActiveEditorEntityManager, Handle, Callback);
 	}
 	else
 	{
-		Result.Completed = FQueryResult::ECompletion::Unavailable;
+		return FQueryResult();
 	}
-
-	return Result;
 }
 
 FTypedElementOnDataStorageUpdate& UTypedElementDatabase::OnUpdate()
@@ -1023,26 +664,17 @@ bool UTypedElementDatabase::ColumnsToBitSets(TConstArrayView<const UScriptStruct
 
 void UTypedElementDatabase::PreparePhase(EQueryTickPhase Phase, float DeltaTime)
 {
-	PhasePreOrPostAmble(Phase, DeltaTime, PhasePreparationQueries[static_cast<std::underlying_type_t<EQueryTickPhase>>(Phase)]);
+	if (ActiveEditorEntityManager)
+	{
+		Queries.RunPhasePreambleQueries(*ActiveEditorEntityManager, Phase, DeltaTime);
+	}
 }
 
 void UTypedElementDatabase::FinalizePhase(EQueryTickPhase Phase, float DeltaTime)
 {
-	PhasePreOrPostAmble(Phase, DeltaTime, PhaseFinalizationQueries[static_cast<std::underlying_type_t<EQueryTickPhase>>(Phase)]);
-}
-
-void UTypedElementDatabase::PhasePreOrPostAmble(EQueryTickPhase Phase, float DeltaTime, TArray<TypedElementQueryHandle>& QueryHandles)
-{
-	if (ActiveEditorEntityManager && !QueryHandles.IsEmpty())
+	if (ActiveEditorEntityManager)
 	{
-		FPhasePreOrPostAmbleExecutor Executor(*ActiveEditorEntityManager, DeltaTime);
-		for (TypedElementQueryHandle Query : QueryHandles)
-		{
-			QueryStore::Handle Handle;
-			Handle.Handle = Query;
-			FTypedElementDatabaseExtendedQuery& QueryData = Queries.GetMutable(Handle);
-			Executor.ExecuteQuery(QueryData.Description, QueryData.NativeQuery, QueryData.Description.Callback.Function);
-		}
+		Queries.RunPhasePostambleQueries(*ActiveEditorEntityManager, Phase, DeltaTime);
 	}
 }
 
@@ -1055,36 +687,5 @@ void UTypedElementDatabase::Reset()
 
 void UTypedElementDatabase::DebugPrintQueryCallbacks(FOutputDevice& Output)
 {
-	Output.Log(TEXT("The Typed Elements Data Storage has the following query callbacks:"));
-	Queries.ListAliveEntries(
-		[&Output](const FTypedElementDatabaseExtendedQuery& Query)
-		{
-			if (Query.Processor)
-			{
-				Output.Logf(TEXT("    [%s] %s"), 
-					IsValid(Query.Processor.Get()) ? TEXT("Valid") : TEXT("Invalid"),
-					*(Query.Processor->GetProcessorName()));
-			}
-		});
-
-	using PhaseType = std::underlying_type_t<EQueryTickPhase>;
-	for (PhaseType PhaseId = 0; PhaseId < static_cast<PhaseType>(EQueryTickPhase::Max); ++PhaseId)
-	{
-		for (TypedElementQueryHandle QueryHandle : PhasePreparationQueries[PhaseId])
-		{
-			QueryStore::Handle Handle;
-			Handle.Handle = QueryHandle;
-			const FTypedElementDatabaseExtendedQuery& QueryData = Queries.Get(Handle);
-			Output.Logf(TEXT("    [Valid] %s [Editor Phase Preamble]"), *QueryData.Description.Callback.Name.ToString());
-		}
-		for (TypedElementQueryHandle QueryHandle : PhaseFinalizationQueries[PhaseId])
-		{
-			QueryStore::Handle Handle;
-			Handle.Handle = QueryHandle;
-			const FTypedElementDatabaseExtendedQuery& QueryData = Queries.Get(Handle);
-			Output.Logf(TEXT("    [Valid] %s [Editor Phase Postamble]"), *QueryData.Description.Callback.Name.ToString());
-		}
-	}
-
-	Output.Log(TEXT("End of Typed Elements Data Storage query callback list."));
+	Queries.DebugPrintQueryCallbacks(Output);
 }
\ No newline at end of file
diff --git a/Engine/Plugins/Experimental/TypedElementsDataStorage/Source/TypedElementsDataStorage/Private/TypedElementDatabase.h b/Engine/Plugins/Experimental/TypedElementsDataStorage/Source/TypedElementsDataStorage/Private/TypedElementDatabase.h
--- a/Engine/Plugins/Experimental/TypedElementsDataStorage/Source/TypedElementsDataStorage/Private/TypedElementDatabase.h
+++ b/Engine/Plugins/Experimental/TypedElementsDataStorage/Source/TypedElementsDataStorage/Private/TypedElementDatabase.h
@@ -5,9 +5,7 @@
 #include "Templates/SharedPointer.h"
 #include "Elements/Interfaces/TypedElementDataStorageInterface.h"
 #include "MassArchetypeTypes.h"
-#include "MassEntityQuery.h"
-#include "MassProcessor.h"
-#include "TypedElementHandleStore.h"
+#include "Queries/TypedElementExtendedQueryStore.h"
 #include "UObject/ObjectMacros.h"
 #include "UObject/StrongObjectPtr.h"
 
@@ -18,13 +16,6 @@ struct FMassProcessingPhaseManager;
 class FOutputDevice;
 class UWorld;
 
-struct FTypedElementDatabaseExtendedQuery
-{
-	FMassEntityQuery NativeQuery; // Used if there's no processor bound.
-	ITypedElementDataStorageInterface::FQueryDescription Description;
-	TStrongObjectPtr<UMassProcessor> Processor;
-};
-
 UCLASS()
 class TYPEDELEMENTSDATASTORAGE_API UTypedElementDatabase 
 	: public UObject
@@ -91,44 +82,21 @@ public:
 	void* GetExternalSystemAddress(UClass* Target) override;
 
 	void DebugPrintQueryCallbacks(FOutputDevice& Output);
-private:
-	using QueryStore = TTypedElementHandleStore<FTypedElementDatabaseExtendedQuery>;
-
-	struct FTickGroupId
-	{
-		FName Name;
-		EQueryTickPhase Phase;
-
-		friend inline uint32 GetTypeHash(const FTickGroupId& Id){ return HashCombine(GetTypeHash(Id.Name), GetTypeHash(Id.Phase)); }
-		friend inline bool operator==(const FTickGroupId& Lhs, const FTickGroupId& Rhs) { return Lhs.Phase == Rhs.Phase && Lhs.Name == Rhs.Name; }
-		friend inline bool operator!=(const FTickGroupId& Lhs, const FTickGroupId& Rhs) { return Lhs.Phase != Rhs.Phase || Lhs.Name != Rhs.Name; }
-	};
-
-	struct FTickGroupDescription
-	{
-		TArray<FName> BeforeGroups;
-		TArray<FName> AfterGroups;
-		bool bRequiresMainThread{ false };
-	};
-	
+
+private:	
 	/** Converts a set of column types into Mass specific fragment and tag bit sets. Returns true if any values were added. */
 	static bool ColumnsToBitSets(TConstArrayView<const UScriptStruct*> Columns, FMassFragmentBitSet& Fragments, FMassTagBitSet& Tags);
 
 	void PreparePhase(EQueryTickPhase Phase, float DeltaTime);
 	void FinalizePhase(EQueryTickPhase Phase, float DeltaTime);
-	void PhasePreOrPostAmble(EQueryTickPhase Phase, float DeltaTime, TArray<TypedElementQueryHandle>& Queries);
 	void Reset();
 	
 	static const FName TickGroupName_SyncWidget;
 	
 	TArray<FMassArchetypeHandle> Tables;
 	TMap<FName, TypedElementTableHandle> TableNameLookup;
-	TMap<FTickGroupId, FTickGroupDescription> TickGroupDescriptions;
-
-	TArray<TypedElementQueryHandle> PhasePreparationQueries[static_cast<std::underlying_type_t<EQueryTickPhase>>(EQueryTickPhase::Max)];
-	TArray<TypedElementQueryHandle> PhaseFinalizationQueries[static_cast<std::underlying_type_t<EQueryTickPhase>>(EQueryTickPhase::Max)];
 
-	QueryStore Queries;
+	FTypedElementExtendedQueryStore Queries;
 
 	FTypedElementOnDataStorageUpdate OnUpdateDelegate;
 
diff --git a/Engine/Plugins/Experimental/TypedElementsDataStorage/Source/TypedElementsDataStorage/Private/TypedElementHandleStore.h b/Engine/Plugins/Experimental/TypedElementsDataStorage/Source/TypedElementsDataStorage/Private/TypedElementHandleStore.h
--- a/Engine/Plugins/Experimental/TypedElementsDataStorage/Source/TypedElementsDataStorage/Private/TypedElementHandleStore.h
+++ b/Engine/Plugins/Experimental/TypedElementsDataStorage/Source/TypedElementsDataStorage/Private/TypedElementHandleStore.h
@@ -27,6 +27,8 @@ public:
 		FHandleData Data;
 		uint64 Handle;
 	};
+
+	using ListAliveEntriesCallback = TFunctionRef<void(const DataType&)>;
 	
 	template<typename... Args>
 	Handle Emplace(Args... Arguments);
@@ -39,7 +41,14 @@ public:
 	void Remove(Handle Entry);
 
 	bool IsAlive(Handle Entry) const;
-	void ListAliveEntries(const TFunctionRef<void(const DataType&)>& Callback) const;
+	void ListAliveEntries(const ListAliveEntriesCallback& Callback) const;
+
+	friend bool operator==(Handle Lhs, Handle Rhs) { return Lhs.Handle == Rhs.Handle; }
+	friend bool operator!=(Handle Lhs, Handle Rhs) { return Lhs.Handle != Rhs.Handle; }
+	friend bool operator<=(Handle Lhs, Handle Rhs) { return Lhs.Handle <= Rhs.Handle; }
+	friend bool operator>=(Handle Lhs, Handle Rhs) { return Lhs.Handle >= Rhs.Handle; }
+	friend bool operator< (Handle Lhs, Handle Rhs) { return Lhs.Handle <  Rhs.Handle; }
+	friend bool operator> (Handle Lhs, Handle Rhs) { return Lhs.Handle >  Rhs.Handle; }
 
 private:
 	TArray<DataType> Data;
@@ -48,6 +57,7 @@ private:
 };
 
 
+
 // Implementations
 
 template<typename DataType, uint32 ReservationSize>
@@ -127,7 +137,7 @@ bool TTypedElementHandleStore<DataType, ReservationSize>::IsAlive(Handle Entry)
 }
 
 template<typename DataType, uint32 ReservationSize>
-void TTypedElementHandleStore<DataType, ReservationSize>::ListAliveEntries(const TFunctionRef<void(const DataType&)>& Callback) const
+void TTypedElementHandleStore<DataType, ReservationSize>::ListAliveEntries(const ListAliveEntriesCallback& Callback) const
 {
 	int32 Count = Data.Num();
 	const DataType* EntryIt = Data.GetData();
