Renamed several "*Processor" files to *Queries" to be inline with recent TEDS changes.

#nojira

diff --git a/Engine/Plugins/Experimental/TypedElementsDataStorage/Source/TypedElementsDataStorage/Private/Compatibility/TypedElementActorLabelProcessors.cpp b/Engine/Plugins/Experimental/TypedElementsDataStorage/Source/TypedElementsDataStorage/Private/Compatibility/TypedElementActorLabelProcessors.cpp
deleted file mode 100644
diff --git a/Engine/Plugins/Experimental/TypedElementsDataStorage/Source/TypedElementsDataStorage/Private/Compatibility/TypedElementActorLabelProcessors.h b/Engine/Plugins/Experimental/TypedElementsDataStorage/Source/TypedElementsDataStorage/Private/Compatibility/TypedElementActorLabelProcessors.h
deleted file mode 100644
diff --git a/Engine/Plugins/Experimental/TypedElementsDataStorage/Source/TypedElementsDataStorage/Private/Compatibility/TypedElementActorLabelQueries.cpp b/Engine/Plugins/Experimental/TypedElementsDataStorage/Source/TypedElementsDataStorage/Private/Compatibility/TypedElementActorLabelQueries.cpp
new file mode 100644
--- /dev/null
+++ b/Engine/Plugins/Experimental/TypedElementsDataStorage/Source/TypedElementsDataStorage/Private/Compatibility/TypedElementActorLabelQueries.cpp
@@ -0,0 +1,135 @@
+// Copyright Epic Games, Inc. All Rights Reserved.
+
+#include "Compatibility/TypedElementActorLabelQueries.h"
+
+#include "Elements/Columns/TypedElementMiscColumns.h"
+#include "Elements/Columns/TypedElementLabelColumns.h"
+#include "Elements/Framework/TypedElementQueryBuilder.h"
+#include "Elements/Framework/TypedElementRegistry.h"
+#include "HAL/IConsoleManager.h"
+#include "Hash/CityHash.h"
+#include "MassActorSubsystem.h"
+#include "ProfilingDebugging/CpuProfilerTrace.h"
+
+FAutoConsoleCommandWithOutputDevice PrintActorLabelsConsoleCommand(
+	TEXT("TEDS.PrintActorLabels"),
+	TEXT("Prints out the labels for all actors found in the Typed Elements Data Storage."),
+	FConsoleCommandWithOutputDeviceDelegate::CreateLambda([](FOutputDevice& Output)
+		{
+			using namespace TypedElementQueryBuilder;
+			using DSI = ITypedElementDataStorageInterface;
+
+			TRACE_CPUPROFILER_EVENT_SCOPE(TEDS.PrintActorLabelsCommand);
+
+			if (ITypedElementDataStorageInterface* DataStorage = UTypedElementRegistry::GetInstance()->GetMutableDataStorage())
+			{
+				static TypedElementQueryHandle LabelQuery = TypedElementInvalidQueryHandle;
+				if (LabelQuery == TypedElementInvalidQueryHandle)
+				{
+					LabelQuery = DataStorage->RegisterQuery(
+						Select()
+							.ReadOnly<FTypedElementLabelColumn>()
+						.Where()
+							.All<FMassActorFragment>()
+						.Compile());
+				}
+				
+				if (LabelQuery != TypedElementInvalidQueryHandle)
+				{
+					FString Message;
+					Output.Log(TEXT("The Typed Elements Data Storage has the following actors:"));
+					DataStorage->RunQuery(LabelQuery, [&Output, &Message](const DSI::FQueryDescription&, DSI::IDirectQueryContext& Context)
+						{
+							const FTypedElementLabelColumn* Labels = Context.GetColumn<FTypedElementLabelColumn>();
+							const uint32 Count = Context.GetRowCount();
+
+							const FTypedElementLabelColumn* LabelsIt = Labels;
+							int32 CharacterCount = 0;
+							// Reserve memory first to avoid repeated memory allocations.
+							for (uint32 Index = 0; Index < Count; ++Index)
+							{
+								CharacterCount += 12 /*Prefixed text size*/ + LabelsIt->Label.Len() + 1 /*Trailing new line*/;
+								++LabelsIt;
+							}
+							Message.Reset(CharacterCount);
+
+							LabelsIt = Labels;
+							for (uint32 Index = 0; Index < Count; ++Index)
+							{
+								Message += TEXT("    Actor: ");
+								Message += LabelsIt->Label;
+								Message += TEXT('\n');
+								++LabelsIt;
+							}
+
+							Output.Log(Message);
+						});
+					Output.Log(TEXT("End of Typed Elements Data Storage actors list."));
+				}
+			}
+		}));
+
+void UTypedElementActorLabelFactory::RegisterQueries(ITypedElementDataStorageInterface& DataStorage) const
+{
+	RegisterActorLabelToColumnQuery(DataStorage);
+	RegisterLabelColumnToActorQuery(DataStorage);
+}
+
+void UTypedElementActorLabelFactory::RegisterActorLabelToColumnQuery(ITypedElementDataStorageInterface& DataStorage) const
+{
+	using namespace TypedElementQueryBuilder;
+	using DSI = ITypedElementDataStorageInterface;
+
+	DataStorage.RegisterQuery(
+		Select(
+			TEXT("Sync actor label to column"),
+			FProcessor(DSI::EQueryTickPhase::PrePhysics, DataStorage.GetQueryTickGroupName(DSI::EQueryTickGroups::SyncExternalToDataStorage))
+				.ForceToGameThread(true),
+			[](const FMassActorFragment& Actor, FTypedElementLabelColumn& Label, FTypedElementLabelHashColumn& LabelHash)
+			{
+				if (const AActor* ActorInstance = Actor.Get(); ActorInstance != nullptr)
+				{
+					const FString& ActorLabel = ActorInstance->GetActorLabel(false);
+					uint64 ActorLabelHash = CityHash64(reinterpret_cast<const char*>(*ActorLabel), ActorLabel.Len() * sizeof(**ActorLabel));
+					if (LabelHash.LabelHash != ActorLabelHash)
+					{
+						Label.Label = ActorLabel;
+						LabelHash.LabelHash = ActorLabelHash;
+					}
+				}
+			}
+		)
+		.Where()
+			.All<FTypedElementSyncFromWorldTag>()
+		.Compile()
+	);
+}
+
+void UTypedElementActorLabelFactory::RegisterLabelColumnToActorQuery(ITypedElementDataStorageInterface& DataStorage) const
+{
+	using namespace TypedElementQueryBuilder;
+	using DSI = ITypedElementDataStorageInterface;
+
+	DataStorage.RegisterQuery(
+		Select(
+			TEXT("Sync label column to actor"),
+			FProcessor(DSI::EQueryTickPhase::FrameEnd, DataStorage.GetQueryTickGroupName(DSI::EQueryTickGroups::SyncDataStorageToExternal))
+				.ForceToGameThread(true),
+			[](FMassActorFragment& Actor, const FTypedElementLabelColumn& Label, const FTypedElementLabelHashColumn& LabelHash)
+			{
+				if (AActor* ActorInstance = Actor.GetMutable(); ActorInstance != nullptr)
+				{
+					const FString& ActorLabel = ActorInstance->GetActorLabel(false);
+					uint64 ActorLabelHash = CityHash64(reinterpret_cast<const char*>(*ActorLabel), ActorLabel.Len() * sizeof(**ActorLabel));
+					if (LabelHash.LabelHash != ActorLabelHash)
+					{
+						ActorInstance->SetActorLabel(Label.Label);
+					}
+				}
+			}
+		)
+		.Where()
+			.All<FTypedElementSyncBackToWorldTag>()
+		.Compile()
+	);
+}
diff --git a/Engine/Plugins/Experimental/TypedElementsDataStorage/Source/TypedElementsDataStorage/Private/Compatibility/TypedElementActorLabelQueries.h b/Engine/Plugins/Experimental/TypedElementsDataStorage/Source/TypedElementsDataStorage/Private/Compatibility/TypedElementActorLabelQueries.h
new file mode 100644
--- /dev/null
+++ b/Engine/Plugins/Experimental/TypedElementsDataStorage/Source/TypedElementsDataStorage/Private/Compatibility/TypedElementActorLabelQueries.h
@@ -0,0 +1,31 @@
+// Copyright Epic Games, Inc. All Rights Reserved.
+
+#pragma once
+
+#include "Elements/Interfaces/TypedElementDataStorageFactory.h"
+#include "UObject/ObjectMacros.h"
+
+#include "TypedElementActorLabelQueries.generated.h"
+
+
+UCLASS()
+class UTypedElementActorLabelFactory : public UTypedElementDataStorageFactory
+{
+	GENERATED_BODY()
+
+public:
+	~UTypedElementActorLabelFactory() override = default;
+
+	void RegisterQueries(ITypedElementDataStorageInterface& DataStorage) const override;
+
+private:
+	/**
+	 * Takes the label set on an actor and copies it to the Data Storage if they differ.
+	 */
+	void RegisterActorLabelToColumnQuery(ITypedElementDataStorageInterface& DataStorage) const;
+	/**
+	 * Takes the label stored in the Data Storage and copies it to the actor's label if the FTypedElementSyncBackToWorldTag
+	 * has been set and the labels differ.
+	 */
+	void RegisterLabelColumnToActorQuery(ITypedElementDataStorageInterface& DataStorage) const;
+};
\ No newline at end of file
diff --git a/Engine/Plugins/Experimental/TypedElementsDataStorage/Source/TypedElementsDataStorage/Private/Compatibility/TypedElementActorPackagePathToColumnProcessor.cpp b/Engine/Plugins/Experimental/TypedElementsDataStorage/Source/TypedElementsDataStorage/Private/Compatibility/TypedElementActorPackagePathToColumnProcessor.cpp
deleted file mode 100644
diff --git a/Engine/Plugins/Experimental/TypedElementsDataStorage/Source/TypedElementsDataStorage/Private/Compatibility/TypedElementActorPackagePathToColumnProcessor.h b/Engine/Plugins/Experimental/TypedElementsDataStorage/Source/TypedElementsDataStorage/Private/Compatibility/TypedElementActorPackagePathToColumnProcessor.h
deleted file mode 100644
diff --git a/Engine/Plugins/Experimental/TypedElementsDataStorage/Source/TypedElementsDataStorage/Private/Compatibility/TypedElementActorPackagePathToColumnQueries.cpp b/Engine/Plugins/Experimental/TypedElementsDataStorage/Source/TypedElementsDataStorage/Private/Compatibility/TypedElementActorPackagePathToColumnQueries.cpp
new file mode 100644
--- /dev/null
+++ b/Engine/Plugins/Experimental/TypedElementsDataStorage/Source/TypedElementsDataStorage/Private/Compatibility/TypedElementActorPackagePathToColumnQueries.cpp
@@ -0,0 +1,32 @@
+// Copyright Epic Games, Inc. All Rights Reserved.
+
+#include "Compatibility/TypedElementActorPackagePathToColumnQueries.h"
+
+#include "Elements/Columns/TypedElementPackageColumns.h"
+#include "Elements/Framework/TypedElementQueryBuilder.h"
+#include "MassActorSubsystem.h"
+#include "UObject/Package.h"
+
+void UTypedElementActorPackagePathFactory::RegisterQueries(ITypedElementDataStorageInterface& DataStorage) const
+{
+	using namespace TypedElementQueryBuilder;
+	using DSI = ITypedElementDataStorageInterface;
+
+	DataStorage.RegisterQuery(
+		Select(
+			TEXT("Sync actor package info to columns"),
+			FObserver(FObserver::EEvent::Add, FMassActorFragment::StaticStruct())
+				.ForceToGameThread(true),
+			[](const FMassActorFragment& Actor, FTypedElementPackagePathColumn& Path, FTypedElementPackageLoadedPathColumn& LoadedPath)
+			{
+				if (const AActor* ActorInstance = Actor.Get(); ActorInstance != nullptr)
+				{
+					const UPackage* Target = ActorInstance->GetPackage();
+					Target->GetPathName(nullptr, Path.Path);
+					LoadedPath.LoadedPath = Target->GetLoadedPath();
+				}
+			}
+		)
+		.Compile()
+	);
+}
diff --git a/Engine/Plugins/Experimental/TypedElementsDataStorage/Source/TypedElementsDataStorage/Private/Compatibility/TypedElementActorPackagePathToColumnQueries.h b/Engine/Plugins/Experimental/TypedElementsDataStorage/Source/TypedElementsDataStorage/Private/Compatibility/TypedElementActorPackagePathToColumnQueries.h
new file mode 100644
--- /dev/null
+++ b/Engine/Plugins/Experimental/TypedElementsDataStorage/Source/TypedElementsDataStorage/Private/Compatibility/TypedElementActorPackagePathToColumnQueries.h
@@ -0,0 +1,19 @@
+// Copyright Epic Games, Inc. All Rights Reserved.
+
+#pragma once
+
+#include "Elements/Interfaces/TypedElementDataStorageFactory.h"
+#include "UObject/ObjectMacros.h"
+
+#include "TypedElementActorPackagePathToColumnQueries.generated.h"
+
+UCLASS()
+class UTypedElementActorPackagePathFactory : public UTypedElementDataStorageFactory
+{
+	GENERATED_BODY()
+
+public:
+	~UTypedElementActorPackagePathFactory() override = default;
+
+	void RegisterQueries(ITypedElementDataStorageInterface& DataStorage) const override;
+};
diff --git a/Engine/Plugins/Experimental/TypedElementsDataStorage/Source/TypedElementsDataStorage/Private/Compatibility/TypedElementActorTransformProcessors.cpp b/Engine/Plugins/Experimental/TypedElementsDataStorage/Source/TypedElementsDataStorage/Private/Compatibility/TypedElementActorTransformProcessors.cpp
deleted file mode 100644
diff --git a/Engine/Plugins/Experimental/TypedElementsDataStorage/Source/TypedElementsDataStorage/Private/Compatibility/TypedElementActorTransformProcessors.h b/Engine/Plugins/Experimental/TypedElementsDataStorage/Source/TypedElementsDataStorage/Private/Compatibility/TypedElementActorTransformProcessors.h
deleted file mode 100644
diff --git a/Engine/Plugins/Experimental/TypedElementsDataStorage/Source/TypedElementsDataStorage/Private/Compatibility/TypedElementActorTransformQueries.cpp b/Engine/Plugins/Experimental/TypedElementsDataStorage/Source/TypedElementsDataStorage/Private/Compatibility/TypedElementActorTransformQueries.cpp
new file mode 100644
--- /dev/null
+++ b/Engine/Plugins/Experimental/TypedElementsDataStorage/Source/TypedElementsDataStorage/Private/Compatibility/TypedElementActorTransformQueries.cpp
@@ -0,0 +1,94 @@
+// Copyright Epic Games, Inc. All Rights Reserved.
+
+#include "Compatibility/TypedElementActorTransformQueries.h"
+
+#include "Elements/Columns/TypedElementMiscColumns.h"
+#include "Elements/Columns/TypedElementTransformColumns.h"
+#include "Elements/Framework/TypedElementQueryBuilder.h"
+#include "MassActorSubsystem.h"
+
+void UTypedElementActorTransformFactory::RegisterQueries(ITypedElementDataStorageInterface& DataStorage) const
+{
+	RegisterActorAddTransformColumn(DataStorage);
+	RegisterActorLocalTransformToColumn(DataStorage);
+	RegisterLocalTransformColumnToActor(DataStorage);
+}
+
+void UTypedElementActorTransformFactory::RegisterActorAddTransformColumn(ITypedElementDataStorageInterface& DataStorage) const
+{
+	using namespace TypedElementQueryBuilder;
+	using DSI = ITypedElementDataStorageInterface;
+
+	DataStorage.RegisterQuery(
+		Select(
+			TEXT("Add transform column to actor"),
+			FProcessor(DSI::EQueryTickPhase::PrePhysics, 
+				DataStorage.GetQueryTickGroupName(DSI::EQueryTickGroups::SyncExternalToDataStorage))
+				.ForceToGameThread(true),
+			[](DSI::IQueryContext& Context, TypedElementRowHandle Row, const FMassActorFragment& Actor)
+			{
+				const AActor* ActorInstance = Actor.Get();
+				if (ActorInstance != nullptr && ActorInstance->GetRootComponent())
+				{
+					Context.AddColumns<FTypedElementLocalTransformColumn>(Row);
+				}
+			}
+		)
+		.Where()
+			.None<FTypedElementLocalTransformColumn>()
+		.Compile()
+	);
+}
+
+void UTypedElementActorTransformFactory::RegisterActorLocalTransformToColumn(ITypedElementDataStorageInterface& DataStorage) const
+{
+	using namespace TypedElementQueryBuilder;
+	using DSI = ITypedElementDataStorageInterface;
+
+	DataStorage.RegisterQuery(
+		Select(
+			TEXT("Sync actor transform to column"),
+			FProcessor(DSI::EQueryTickPhase::PostPhysics, DataStorage.GetQueryTickGroupName(DSI::EQueryTickGroups::SyncExternalToDataStorage))
+				.ForceToGameThread(true),
+			[](DSI::IQueryContext& Context, TypedElementRowHandle Row, const FMassActorFragment& Actor, FTypedElementLocalTransformColumn& Transform)
+			{
+				const AActor* ActorInstance = Actor.Get();
+				if (ActorInstance != nullptr && ActorInstance->GetRootComponent() != nullptr)
+				{
+					Transform.Transform = ActorInstance->GetActorTransform();
+				}
+				else
+				{
+					Context.RemoveColumns<FTypedElementLocalTransformColumn>(Row);
+				}
+			}
+		)
+		.Where()
+			.All<FTypedElementSyncFromWorldTag>()
+		.Compile()
+	);
+}
+
+void UTypedElementActorTransformFactory::RegisterLocalTransformColumnToActor(ITypedElementDataStorageInterface& DataStorage) const
+{
+	using namespace TypedElementQueryBuilder;
+	using DSI = ITypedElementDataStorageInterface;
+
+	DataStorage.RegisterQuery(
+		Select(
+			TEXT("Sync transform column to actor"),
+			FProcessor(DSI::EQueryTickPhase::FrameEnd, DataStorage.GetQueryTickGroupName(DSI::EQueryTickGroups::SyncDataStorageToExternal))
+				.ForceToGameThread(true),
+			[](FMassActorFragment& Actor, const FTypedElementLocalTransformColumn& Transform)
+			{
+				if (AActor* ActorInstance = Actor.GetMutable(); ActorInstance != nullptr)
+				{
+					ActorInstance->SetActorTransform(Transform.Transform);
+				}
+			}
+		)
+		.Where()
+			.All< FTypedElementSyncBackToWorldTag>()
+		.Compile()
+	);
+}
diff --git a/Engine/Plugins/Experimental/TypedElementsDataStorage/Source/TypedElementsDataStorage/Private/Compatibility/TypedElementActorTransformQueries.h b/Engine/Plugins/Experimental/TypedElementsDataStorage/Source/TypedElementsDataStorage/Private/Compatibility/TypedElementActorTransformQueries.h
new file mode 100644
--- /dev/null
+++ b/Engine/Plugins/Experimental/TypedElementsDataStorage/Source/TypedElementsDataStorage/Private/Compatibility/TypedElementActorTransformQueries.h
@@ -0,0 +1,36 @@
+// Copyright Epic Games, Inc. All Rights Reserved.
+
+#pragma once
+
+#include "Elements/Interfaces/TypedElementDataStorageFactory.h"
+#include "UObject/ObjectMacros.h"
+
+#include "TypedElementActorTransformQueries.generated.h"
+
+UCLASS()
+class UTypedElementActorTransformFactory : public UTypedElementDataStorageFactory
+{
+	GENERATED_BODY()
+
+public:
+	~UTypedElementActorTransformFactory() override = default;
+
+	void RegisterQueries(ITypedElementDataStorageInterface& DataStorage) const override;
+
+private:
+	/**
+	 * Checks actors that don't have a tranform column and adds one if an actor has been
+	 * assigned a transform.
+	 */
+	void RegisterActorAddTransformColumn(ITypedElementDataStorageInterface& DataStorage) const;
+	/**
+	 * Takes the transform set on an actor and copies it to the Data Storage or removes the
+	 * transform column if there's not transform available anymore.
+	 */
+	void RegisterActorLocalTransformToColumn(ITypedElementDataStorageInterface& DataStorage) const;
+	/**
+	 * Takes the transform stored in the Data Storage and copies it to the actor's tranform if
+	 * the FTypedElementSyncBackToWorldTag has been set.
+	 */
+	void RegisterLocalTransformColumnToActor(ITypedElementDataStorageInterface& DataStorage) const;
+};
diff --git a/Engine/Plugins/Experimental/TypedElementsDataStorage/Source/TypedElementsDataStorage/Private/Processors/TypedElementMiscProcessors.cpp b/Engine/Plugins/Experimental/TypedElementsDataStorage/Source/TypedElementsDataStorage/Private/Processors/TypedElementMiscProcessors.cpp
deleted file mode 100644
diff --git a/Engine/Plugins/Experimental/TypedElementsDataStorage/Source/TypedElementsDataStorage/Private/Processors/TypedElementMiscProcessors.h b/Engine/Plugins/Experimental/TypedElementsDataStorage/Source/TypedElementsDataStorage/Private/Processors/TypedElementMiscProcessors.h
deleted file mode 100644
diff --git a/Engine/Plugins/Experimental/TypedElementsDataStorage/Source/TypedElementsDataStorage/Private/Queries/TypedElementMiscQueries.cpp b/Engine/Plugins/Experimental/TypedElementsDataStorage/Source/TypedElementsDataStorage/Private/Queries/TypedElementMiscQueries.cpp
new file mode 100644
--- /dev/null
+++ b/Engine/Plugins/Experimental/TypedElementsDataStorage/Source/TypedElementsDataStorage/Private/Queries/TypedElementMiscQueries.cpp
@@ -0,0 +1,40 @@
+// Copyright Epic Games, Inc. All Rights Reserved.
+
+#include "Queries/TypedElementMiscQueries.h"
+
+#include "Elements/Columns/TypedElementMiscColumns.h"
+#include "Elements/Framework/TypedElementQueryBuilder.h"
+
+void UTypedElementRemoveSyncToWorldTagFactory::RegisterQueries(ITypedElementDataStorageInterface& DataStorage) const
+{
+	using namespace TypedElementQueryBuilder;
+	using DSI = ITypedElementDataStorageInterface;
+
+	DataStorage.RegisterQuery(
+		Select(
+			TEXT("Remove 'sync to world' tag"),
+			FPhaseAmble(FPhaseAmble::ELocation::Postamble, DSI::EQueryTickPhase::FrameEnd),
+			[](DSI::IQueryContext& Context, const TypedElementRowHandle* Rows)
+			{
+				Context.RemoveColumns<FTypedElementSyncBackToWorldTag>(TConstArrayView<TypedElementRowHandle>(Rows, Context.GetRowCount()));
+			}
+		)
+		.Where()
+			.All<FTypedElementSyncBackToWorldTag>()
+		.Compile()
+	);
+
+	DataStorage.RegisterQuery(
+	Select(
+		TEXT("Remove 'sync from world' tag"),
+		FPhaseAmble(FPhaseAmble::ELocation::Postamble, DSI::EQueryTickPhase::FrameEnd),
+		[](DSI::IQueryContext& Context, const TypedElementRowHandle* Rows)
+		{
+			Context.RemoveColumns<FTypedElementSyncFromWorldTag>(TConstArrayView<TypedElementRowHandle>(Rows, Context.GetRowCount()));
+		}
+	)
+	.Where()
+		.All<FTypedElementSyncFromWorldTag>()
+	.Compile()
+);
+}
diff --git a/Engine/Plugins/Experimental/TypedElementsDataStorage/Source/TypedElementsDataStorage/Private/Queries/TypedElementMiscQueries.h b/Engine/Plugins/Experimental/TypedElementsDataStorage/Source/TypedElementsDataStorage/Private/Queries/TypedElementMiscQueries.h
new file mode 100644
--- /dev/null
+++ b/Engine/Plugins/Experimental/TypedElementsDataStorage/Source/TypedElementsDataStorage/Private/Queries/TypedElementMiscQueries.h
@@ -0,0 +1,22 @@
+// Copyright Epic Games, Inc. All Rights Reserved.
+
+#pragma once
+
+#include "Elements/Interfaces/TypedElementDataStorageFactory.h"
+#include "UObject/ObjectMacros.h"
+
+#include "TypedElementMiscQueries.generated.h"
+
+/**
+ * Removes all FTypedElementSyncBackToWorldTags at the end of an update cycle.
+ */
+UCLASS()
+class UTypedElementRemoveSyncToWorldTagFactory : public UTypedElementDataStorageFactory
+{
+	GENERATED_BODY()
+
+public:
+	~UTypedElementRemoveSyncToWorldTagFactory() override = default;
+
+	void RegisterQueries(ITypedElementDataStorageInterface& DataStorage) const override;
+};
