Added a few cell-like widgets for use in the Scene Outliner, among others.

diff --git a/Engine/Plugins/Experimental/TypedElementsDataStorage/Source/TypedElementsDataStorageUI/Private/TypedElementLabelWidget.cpp b/Engine/Plugins/Experimental/TypedElementsDataStorage/Source/TypedElementsDataStorageUI/Private/TypedElementLabelWidget.cpp
new file mode 100644
--- /dev/null
+++ b/Engine/Plugins/Experimental/TypedElementsDataStorage/Source/TypedElementsDataStorageUI/Private/TypedElementLabelWidget.cpp
@@ -0,0 +1,196 @@
+// Copyright Epic Games, Inc. All Rights Reserved.
+
+#include "TypedElementLabelWidget.h"
+
+#include "Elements/Columns/TypedElementLabelColumns.h"
+#include "Elements/Columns/TypedElementMiscColumns.h"
+#include "Elements/Columns/TypedElementSlateWidgetColumns.h"
+#include "Elements/Columns/TypedElementValueCacheColumns.h"
+#include "Elements/Framework/TypedElementQueryBuilder.h"
+#include "TypedElementSubsystems.h"
+#include "Widgets/Text/STextBlock.h"
+
+#define LOCTEXT_NAMESPACE "TypedElementUI_LabelWidget"
+
+//
+// UTypedElementLabelWidgetFactory
+//
+
+static void UpdateTextWidget(const TWeakPtr<SWidget>& Widget, const FTypedElementLabelColumn& Label, const uint64* HashValue)
+{
+	TSharedPtr<SWidget> WidgetPointer = Widget.Pin();
+	checkf(WidgetPointer, TEXT("Referenced widget is not valid. A constructed widget may not have been cleaned up. This can "
+		"also happen if this processor is running in the same phase as the processors responsible for cleaning up old "
+		"references."));
+	checkf(WidgetPointer->GetType() == STextBlock::StaticWidgetClass().GetWidgetType(),
+		TEXT("Stored widget with FTypedElementLabelWidgetConstructor doesn't match type %s, but was a %s."),
+		*(STextBlock::StaticWidgetClass().GetWidgetType().ToString()),
+		*(WidgetPointer->GetTypeAsString()));
+	STextBlock* WidgetInstance = static_cast<STextBlock*>(WidgetPointer.Get());
+	WidgetInstance->SetText(FText::FromString(Label.Label));
+	if (!HashValue)
+	{
+		WidgetInstance->SetToolTipText(FText::FromString(Label.Label));
+	}
+	else
+	{
+		WidgetInstance->SetToolTipText(FText::FromString(FString::Format(TEXT("{0}\nHash: {1}"), { Label.Label, *HashValue })));
+	}
+}
+
+static void SyncColumnsToWidget(
+	ITypedElementDataStorageInterface* DataStorage, 
+	TypedElementRowHandle TargetRow, 
+	FTypedElementU64IntValueCacheColumn& TextHash,
+	const TWeakPtr<SWidget>& Widget,
+	bool bShowHashInToolTip)
+{
+	if (const FTypedElementLabelColumn* LabelColumn = DataStorage->GetColumn<FTypedElementLabelColumn>(TargetRow))
+	{
+		if (const FTypedElementLabelHashColumn* LabelHashColumn = DataStorage->GetColumn<FTypedElementLabelHashColumn>(TargetRow))
+		{
+			if (LabelHashColumn->LabelHash != TextHash.Value)
+			{
+				UpdateTextWidget(Widget, *LabelColumn, bShowHashInToolTip ? &LabelHashColumn->LabelHash : nullptr);
+				TextHash.Value = LabelHashColumn->LabelHash;
+			}
+		}
+		else
+		{
+			UpdateTextWidget(Widget, *LabelColumn, nullptr);
+		}
+	}
+}
+
+void UTypedElementLabelWidgetFactory::RegisterQueries(ITypedElementDataStorageInterface& DataStorage) const
+{
+	using namespace TypedElementQueryBuilder;
+	using DSI = ITypedElementDataStorageInterface;
+
+	DataStorage.RegisterQuery(
+		Select(
+			TEXT("Sync label to widget"),
+			FProcessor(DSI::EQueryTickPhase::FrameEnd, DataStorage.GetQueryTickGroupName(DSI::EQueryTickGroups::SyncWidgets))
+				.ForceToGameThread(true),
+			[](	FCachedQueryContext<UTypedElementDataStorageSubsystem>& Context, 
+				FTypedElementSlateWidgetReferenceColumn& Widget,
+				FTypedElementU64IntValueCacheColumn& TextHash,
+				const FTypedElementLabelWidgetColumn& Config,
+				const FTypedElementRowReferenceColumn& Target)
+			{
+				DSI* DataStorage = Context.GetCachedMutableDependency<UTypedElementDataStorageSubsystem>().Get();
+				checkf(DataStorage, TEXT("FTypedElementsDataStorageUiModule tried to process widgets before the "
+					"Typed Elements Data Storage interface is available."));
+				
+				if (DataStorage->HasColumns<FTypedElementSyncFromWorldTag>(Target.Row) ||
+					DataStorage->HasColumns<FTypedElementSyncBackToWorldTag>(Target.Row))
+				{
+					SyncColumnsToWidget(DataStorage, Target.Row, TextHash, Widget.Widget, Config.bShowHashInTooltip);
+				}
+			})
+		.Compile()
+	);
+}
+
+void UTypedElementLabelWidgetFactory::RegisterWidgetConstructors(ITypedElementDataStorageInterface& DataStorage,
+	ITypedElementDataStorageUiInterface& DataStorageUi) const
+{
+	 DataStorageUi.RegisterWidgetFactory(FName(TEXT("General.Cell")), FTypedElementLabelWidgetConstructor::StaticStruct(),
+		{ FTypedElementLabelColumn::StaticStruct() });
+	DataStorageUi.RegisterWidgetFactory(FName(TEXT("General.Cell")), FTypedElementLabelWithHashTooltipWidgetConstructor::StaticStruct(),
+		{ FTypedElementLabelColumn::StaticStruct(), FTypedElementLabelHashColumn::StaticStruct() });
+}
+
+
+
+//
+// FTypedElementLabelWidgetConstructor
+//
+
+FTypedElementLabelWidgetConstructor::FTypedElementLabelWidgetConstructor()
+	: Super(FTypedElementLabelWidgetConstructor::StaticStruct())
+{
+}
+
+FTypedElementLabelWidgetConstructor::FTypedElementLabelWidgetConstructor(const UScriptStruct* InTypeInfo)
+	: Super(InTypeInfo)
+{
+}
+
+TConstArrayView<const UScriptStruct*> FTypedElementLabelWidgetConstructor::GetAdditionalColumnsList() const
+{
+	static const TTypedElementColumnTypeList<
+		FTypedElementRowReferenceColumn,
+		FTypedElementU64IntValueCacheColumn,
+		FTypedElementLabelWidgetColumn> Columns;
+	return Columns;
+}
+
+bool FTypedElementLabelWidgetConstructor::CanBeReused() const
+{
+	return true;
+}
+
+TSharedPtr<SWidget> FTypedElementLabelWidgetConstructor::CreateWidget()
+{
+	return SNew(STextBlock);
+}
+
+bool FTypedElementLabelWidgetConstructor::SetColumns(ITypedElementDataStorageInterface* DataStorage, TypedElementRowHandle Row)
+{
+	DataStorage->GetColumn<FTypedElementU64IntValueCacheColumn>(Row)->Value = 0;
+	DataStorage->GetColumn<FTypedElementLabelWidgetColumn>(Row)->bShowHashInTooltip = false;
+
+	return true;
+}
+
+bool FTypedElementLabelWidgetConstructor::FinalizeWidget(
+	ITypedElementDataStorageInterface* DataStorage,
+	ITypedElementDataStorageUiInterface* DataStorageUi,
+	TypedElementRowHandle Row,
+	const TSharedPtr<SWidget>& Widget)
+{
+	SyncColumnsToWidget(
+		DataStorage, 
+		DataStorage->GetColumn<FTypedElementRowReferenceColumn>(Row)->Row,
+		*DataStorage->GetColumn<FTypedElementU64IntValueCacheColumn>(Row),
+		Widget,
+		false);
+	return true;
+}
+
+
+
+//
+// FTypedElementLabelWithHashTooltipWidgetConstructor
+//
+
+FTypedElementLabelWithHashTooltipWidgetConstructor::FTypedElementLabelWithHashTooltipWidgetConstructor()
+	: Super(FTypedElementLabelWithHashTooltipWidgetConstructor::StaticStruct())
+{
+}
+
+bool FTypedElementLabelWithHashTooltipWidgetConstructor::SetColumns(ITypedElementDataStorageInterface* DataStorage, TypedElementRowHandle Row)
+{
+	DataStorage->GetColumn<FTypedElementU64IntValueCacheColumn>(Row)->Value = 0;
+	DataStorage->GetColumn<FTypedElementLabelWidgetColumn>(Row)->bShowHashInTooltip = true;
+
+	return true;
+}
+
+bool FTypedElementLabelWithHashTooltipWidgetConstructor::FinalizeWidget(
+	ITypedElementDataStorageInterface* DataStorage,
+	ITypedElementDataStorageUiInterface* DataStorageUi,
+	TypedElementRowHandle Row,
+	const TSharedPtr<SWidget>& Widget)
+{
+	SyncColumnsToWidget(
+		DataStorage,
+		DataStorage->GetColumn<FTypedElementRowReferenceColumn>(Row)->Row,
+		*DataStorage->GetColumn<FTypedElementU64IntValueCacheColumn>(Row),
+		Widget,
+		true);
+	return true;
+}
+
+#undef LOCTEXT_NAMESPACE
diff --git a/Engine/Plugins/Experimental/TypedElementsDataStorage/Source/TypedElementsDataStorageUI/Private/TypedElementPackagePathWidget.cpp b/Engine/Plugins/Experimental/TypedElementsDataStorage/Source/TypedElementsDataStorageUI/Private/TypedElementPackagePathWidget.cpp
new file mode 100644
--- /dev/null
+++ b/Engine/Plugins/Experimental/TypedElementsDataStorage/Source/TypedElementsDataStorageUI/Private/TypedElementPackagePathWidget.cpp
@@ -0,0 +1,105 @@
+// Copyright Epic Games, Inc. All Rights Reserved.
+
+#include "TypedElementPackagePathWidget.h"
+
+#include "Elements/Columns/TypedElementPackageColumns.h"
+#include "Elements/Columns/TypedElementMiscColumns.h"
+#include "Elements/Columns/TypedElementSlateWidgetColumns.h"
+#include "Elements/Columns/TypedElementValueCacheColumns.h"
+#include "Elements/Framework/TypedElementQueryBuilder.h"
+#include "TypedElementSubsystems.h"
+#include "Widgets/Text/STextBlock.h"
+
+//
+// UTypedElementPackagePathWidgetFactory
+//
+
+void UTypedElementPackagePathWidgetFactory::RegisterWidgetConstructors(ITypedElementDataStorageInterface& DataStorage,
+	ITypedElementDataStorageUiInterface& DataStorageUi) const
+{
+	DataStorageUi.RegisterWidgetFactory(FName(TEXT("General.Cell")), FTypedElementPackagePathWidgetConstructor::StaticStruct(),
+		{ FTypedElementPackagePathColumn::StaticStruct() });
+	DataStorageUi.RegisterWidgetFactory(FName(TEXT("General.Cell")), FTypedElementLoadedPackagePathWidgetConstructor::StaticStruct(),
+		{ FTypedElementPackageLoadedPathColumn::StaticStruct() });
+}
+
+
+
+//
+// FTypedElementPackagePathWidgetConstructor
+//
+
+FTypedElementPackagePathWidgetConstructor::FTypedElementPackagePathWidgetConstructor()
+	: Super(FTypedElementPackagePathWidgetConstructor::StaticStruct())
+{
+}
+
+FTypedElementPackagePathWidgetConstructor::FTypedElementPackagePathWidgetConstructor(const UScriptStruct* InTypeInfo)
+	: Super(FTypedElementPackagePathWidgetConstructor::StaticStruct())
+{
+}
+
+bool FTypedElementPackagePathWidgetConstructor::CanBeReused() const
+{
+	return true;
+}
+
+TSharedPtr<SWidget> FTypedElementPackagePathWidgetConstructor::CreateWidget()
+{
+	return SNew(STextBlock)
+		.OverflowPolicy(ETextOverflowPolicy::Ellipsis)
+		.Justification(ETextJustify::Right);
+}
+
+bool FTypedElementPackagePathWidgetConstructor::FinalizeWidget(
+	ITypedElementDataStorageInterface* DataStorage,
+	ITypedElementDataStorageUiInterface* DataStorageUi,
+	TypedElementRowHandle Row,
+	const TSharedPtr<SWidget>& Widget)
+{
+	TypedElementRowHandle TargetRow = DataStorage->GetColumn<FTypedElementRowReferenceColumn>(Row)->Row;
+	if (const FTypedElementPackagePathColumn* Path = DataStorage->GetColumn<FTypedElementPackagePathColumn>(TargetRow))
+	{
+		STextBlock* TextWidget = static_cast<STextBlock*>(Widget.Get());
+		FText Text = FText::FromString(Path->Path);
+		TextWidget->SetToolTipText(Text);
+		TextWidget->SetText(MoveTemp(Text));
+		return true;
+	}
+	else
+	{
+		return false;
+	}
+}
+
+
+
+//
+// FTypedElementPackagePathWidgetConstructor
+//
+
+FTypedElementLoadedPackagePathWidgetConstructor::FTypedElementLoadedPackagePathWidgetConstructor()
+	: Super(FTypedElementLoadedPackagePathWidgetConstructor::StaticStruct())
+{
+}
+
+bool FTypedElementLoadedPackagePathWidgetConstructor::FinalizeWidget(
+	ITypedElementDataStorageInterface* DataStorage,
+	ITypedElementDataStorageUiInterface* DataStorageUi,
+	TypedElementRowHandle Row,
+	const TSharedPtr<SWidget>& Widget)
+{
+	TypedElementRowHandle TargetRow = DataStorage->GetColumn<FTypedElementRowReferenceColumn>(Row)->Row;
+	if (const FTypedElementPackageLoadedPathColumn* Path = DataStorage->GetColumn<FTypedElementPackageLoadedPathColumn>(TargetRow))
+	{
+		STextBlock* TextWidget = static_cast<STextBlock*>(Widget.Get());
+		FText Text = FText::FromString(Path->LoadedPath.GetLocalFullPath());
+		TextWidget->SetToolTipText(Text);
+		TextWidget->SetText(MoveTemp(Text));
+		return true;
+	}
+	else
+	{
+		return false;
+	}
+}
diff --git a/Engine/Plugins/Experimental/TypedElementsDataStorage/Source/TypedElementsDataStorageUI/Public/TypedElementLabelWidget.h b/Engine/Plugins/Experimental/TypedElementsDataStorage/Source/TypedElementsDataStorageUI/Public/TypedElementLabelWidget.h
new file mode 100644
--- /dev/null
+++ b/Engine/Plugins/Experimental/TypedElementsDataStorage/Source/TypedElementsDataStorageUI/Public/TypedElementLabelWidget.h
@@ -0,0 +1,67 @@
+// Copyright Epic Games, Inc. All Rights Reserved.
+
+#pragma once
+
+#include "Elements/Interfaces/TypedElementDataStorageFactory.h"
+#include "Elements/Interfaces/TypedElementDataStorageInterface.h"
+#include "Elements/Interfaces/TypedElementDataStorageUiInterface.h"
+#include "Internationalization/Text.h"
+#include "UObject/ObjectMacros.h"
+
+#include "TypedElementLabelWidget.generated.h"
+
+UCLASS()
+class TYPEDELEMENTSDATASTORAGEUI_API UTypedElementLabelWidgetFactory : public UTypedElementDataStorageFactory
+{
+	GENERATED_BODY()
+
+public:
+	~UTypedElementLabelWidgetFactory() override = default;
+
+	void RegisterQueries(ITypedElementDataStorageInterface& DataStorage) const override;
+	void RegisterWidgetConstructors(ITypedElementDataStorageInterface& DataStorage,
+		ITypedElementDataStorageUiInterface& DataStorageUi) const override;
+};
+
+USTRUCT()
+struct TYPEDELEMENTSDATASTORAGEUI_API FTypedElementLabelWidgetConstructor : public FTypedElementWidgetConstructor
+{
+	GENERATED_BODY()
+
+public:
+	FTypedElementLabelWidgetConstructor();
+	~FTypedElementLabelWidgetConstructor() override = default;
+
+	TConstArrayView<const UScriptStruct*> GetAdditionalColumnsList() const override;
+	bool CanBeReused() const override;
+
+protected:
+	explicit FTypedElementLabelWidgetConstructor(const UScriptStruct* InTypeInfo);
+	TSharedPtr<SWidget> CreateWidget() override;
+	bool SetColumns(ITypedElementDataStorageInterface* DataStorage, TypedElementRowHandle Row) override;
+	bool FinalizeWidget(ITypedElementDataStorageInterface* DataStorage, ITypedElementDataStorageUiInterface* DataStorageUi,
+		TypedElementRowHandle Row, const TSharedPtr<SWidget>& Widget) override;
+};
+
+USTRUCT()
+struct TYPEDELEMENTSDATASTORAGEUI_API FTypedElementLabelWithHashTooltipWidgetConstructor : public FTypedElementLabelWidgetConstructor
+{
+	GENERATED_BODY()
+
+public:
+	FTypedElementLabelWithHashTooltipWidgetConstructor();
+	~FTypedElementLabelWithHashTooltipWidgetConstructor() = default;
+
+protected:
+	bool SetColumns(ITypedElementDataStorageInterface* DataStorage, TypedElementRowHandle Row) override;
+	bool FinalizeWidget(ITypedElementDataStorageInterface* DataStorage, ITypedElementDataStorageUiInterface* DataStorageUi,
+		TypedElementRowHandle Row, const TSharedPtr<SWidget>& Widget) override;
+};
+
+USTRUCT(meta = (DisplayName = "Label widget"))
+struct TYPEDELEMENTSDATASTORAGEUI_API FTypedElementLabelWidgetColumn : public FTypedElementDataStorageColumn
+{
+	GENERATED_BODY()
+
+	bool bShowHashInTooltip{ false };
+};
\ No newline at end of file
diff --git a/Engine/Plugins/Experimental/TypedElementsDataStorage/Source/TypedElementsDataStorageUI/Public/TypedElementPackagePathWidget.h b/Engine/Plugins/Experimental/TypedElementsDataStorage/Source/TypedElementsDataStorageUI/Public/TypedElementPackagePathWidget.h
new file mode 100644
--- /dev/null
+++ b/Engine/Plugins/Experimental/TypedElementsDataStorage/Source/TypedElementsDataStorageUI/Public/TypedElementPackagePathWidget.h
@@ -0,0 +1,56 @@
+// Copyright Epic Games, Inc. All Rights Reserved.
+
+#pragma once
+
+#include "Elements/Interfaces/TypedElementDataStorageFactory.h"
+#include "Elements/Interfaces/TypedElementDataStorageInterface.h"
+#include "Elements/Interfaces/TypedElementDataStorageUiInterface.h"
+#include "Internationalization/Text.h"
+#include "UObject/ObjectMacros.h"
+
+#include "TypedElementPackagePathWidget.generated.h"
+
+UCLASS()
+class TYPEDELEMENTSDATASTORAGEUI_API UTypedElementPackagePathWidgetFactory : public UTypedElementDataStorageFactory
+{
+	GENERATED_BODY()
+
+public:
+	~UTypedElementPackagePathWidgetFactory() override = default;
+
+	void RegisterWidgetConstructors(ITypedElementDataStorageInterface& DataStorage,
+		ITypedElementDataStorageUiInterface& DataStorageUi) const override;
+};
+
+USTRUCT()
+struct TYPEDELEMENTSDATASTORAGEUI_API FTypedElementPackagePathWidgetConstructor : public FTypedElementWidgetConstructor
+{
+	GENERATED_BODY()
+
+public:
+	FTypedElementPackagePathWidgetConstructor();
+	~FTypedElementPackagePathWidgetConstructor() override = default;
+
+	bool CanBeReused() const override;
+
+protected:
+	explicit FTypedElementPackagePathWidgetConstructor(const UScriptStruct* InTypeInfo);
+
+	TSharedPtr<SWidget> CreateWidget() override;
+	bool FinalizeWidget(ITypedElementDataStorageInterface* DataStorage, ITypedElementDataStorageUiInterface* DataStorageUi,
+		TypedElementRowHandle Row, const TSharedPtr<SWidget>& Widget) override;
+};
+
+USTRUCT()
+struct TYPEDELEMENTSDATASTORAGEUI_API FTypedElementLoadedPackagePathWidgetConstructor : public FTypedElementPackagePathWidgetConstructor
+{
+	GENERATED_BODY()
+
+public:
+	FTypedElementLoadedPackagePathWidgetConstructor();
+	~FTypedElementLoadedPackagePathWidgetConstructor() override = default;
+
+protected:
+	bool FinalizeWidget(ITypedElementDataStorageInterface* DataStorage, ITypedElementDataStorageUiInterface* DataStorageUi,
+		TypedElementRowHandle Row, const TSharedPtr<SWidget>& Widget) override;
+};
diff --git a/Engine/Source/Runtime/TypedElementFramework/Public/Elements/Interfaces/TypedElementDataStorageInterface.h b/Engine/Source/Runtime/TypedElementFramework/Public/Elements/Interfaces/TypedElementDataStorageInterface.h
--- a/Engine/Source/Runtime/TypedElementFramework/Public/Elements/Interfaces/TypedElementDataStorageInterface.h
+++ b/Engine/Source/Runtime/TypedElementFramework/Public/Elements/Interfaces/TypedElementDataStorageInterface.h
@@ -75,7 +75,7 @@ struct TTypedElementColumnTypeList
 {
 	const UScriptStruct* ColumnTypes[sizeof...(Columns)] = { Columns::StaticStruct()... };
 	
-	operator TConstArrayView<const UScriptStruct*>() { return ColumnTypes; }
+	operator TConstArrayView<const UScriptStruct*>() const { return ColumnTypes; }
 };
 
 class TYPEDELEMENTFRAMEWORK_API ITypedElementDataStorageInterface
