Addressed 2 minor issues with TEDS.

The following issues were fixed:
- Some localization strings were colliding. They've not been given unique names.
- Inputing valid types that weren't based on Mass fragments would cause TEDS.UI.BindColumnsToSceneOutliner to assert. A check was added to verify this.

#nojira

diff --git a/Engine/Plugins/Experimental/TypedElementsDataStorage/Source/TypedElementsDataStorageUI/Private/TypedElementGeneralRegistration.cpp b/Engine/Plugins/Experimental/TypedElementsDataStorage/Source/TypedElementsDataStorageUI/Private/TypedElementGeneralRegistration.cpp
--- a/Engine/Plugins/Experimental/TypedElementsDataStorage/Source/TypedElementsDataStorageUI/Private/TypedElementGeneralRegistration.cpp
+++ b/Engine/Plugins/Experimental/TypedElementsDataStorage/Source/TypedElementsDataStorageUI/Private/TypedElementGeneralRegistration.cpp
@@ -13,13 +13,13 @@ const FName UTypedElementGeneralRegistrationFactory::HeaderDefaultPurpose(TEXT("
 void UTypedElementGeneralRegistrationFactory::RegisterWidgetPurposes(ITypedElementDataStorageUiInterface& DataStorageUi) const
 {
 	DataStorageUi.RegisterWidgetPurpose(CellPurpose, ITypedElementDataStorageUiInterface::EPurposeType::UniqueByNameAndColumn,
-		LOCTEXT("TypedElementsUI_GeneralRegistration", "General purpose widgets that can be used as cells for specific columns or column combinations."));
+		LOCTEXT("CellPurpose", "General purpose widgets that can be used as cells for specific columns or column combinations."));
 	
 	DataStorageUi.RegisterWidgetPurpose(HeaderDefaultPurpose, ITypedElementDataStorageUiInterface::EPurposeType::UniqueByName,
-		LOCTEXT("TypedElementsUI_GeneralRegistration", "The default widget to use in headers if no other specialization is provided."));
+		LOCTEXT("HeaderDefaultPurpose", "The default widget to use in headers if no other specialization is provided."));
 	
 	DataStorageUi.RegisterWidgetPurpose(CellDefaultPurpose, ITypedElementDataStorageUiInterface::EPurposeType::UniqueByName,
-		LOCTEXT("TypedElementsUI_GeneralRegistration", "The default widget to use in cells if no other specialization is provided."));
+		LOCTEXT("CellDefaultPurpose", "The default widget to use in cells if no other specialization is provided."));
 }
 
 #undef LOCTEXT_NAMESPACE
diff --git a/Engine/Plugins/Experimental/TypedElementsDataStorage/Source/TypedElementsDataStorageUI/Private/TypedElementSceneOutliner.cpp b/Engine/Plugins/Experimental/TypedElementsDataStorage/Source/TypedElementsDataStorageUI/Private/TypedElementSceneOutliner.cpp
--- a/Engine/Plugins/Experimental/TypedElementsDataStorage/Source/TypedElementsDataStorageUI/Private/TypedElementSceneOutliner.cpp
+++ b/Engine/Plugins/Experimental/TypedElementsDataStorage/Source/TypedElementsDataStorageUI/Private/TypedElementSceneOutliner.cpp
@@ -17,6 +17,7 @@
 #include "ISceneOutlinerColumn.h"
 #include "ISceneOutlinerTreeItem.h"
 #include "LevelEditor.h"
+#include "MassEntityTypes.h"
 #include "SSceneOutliner.h"
 #include "SceneOutlinerFwd.h"
 #include "SceneOutlinerPublicTypes.h"
@@ -81,7 +82,7 @@ FAutoConsoleCommand BindColumnsToSceneOutlinerConsoleCommand(
 								if (!Arg.IsEmpty() && Arg[0] == '/' && Path.TrySetPath(Arg))
 								{
 									const UScriptStruct* ColumnType = TypeOptional(Path);
-									if (ColumnType)
+									if (ColumnType && ColumnType->IsChildOf(FMassFragment::StaticStruct()))
 									{
 										Query.ReadOnly(ColumnType);
 										++AdditionCount;
@@ -248,14 +249,14 @@ void UTypedElementSceneOutlinerFactory::RegisterWidgetPurposes(ITypedElementData
 	using PurposeType = ITypedElementDataStorageUiInterface::EPurposeType;
 
 	DataStorageUi.RegisterWidgetPurpose(FTypedElementSceneOutlinerQueryBinder::HeaderWidgetPurpose, PurposeType::UniqueByNameAndColumn,
-		LOCTEXT("SceneOutlinerCellDescription", "Widgets for headers in any Scene Outliner for specific columns or column combinations."));
+		LOCTEXT("HeaderWidgetPurpose", "Widgets for headers in any Scene Outliner for specific columns or column combinations."));
 	DataStorageUi.RegisterWidgetPurpose(FTypedElementSceneOutlinerQueryBinder::DefaultHeaderWidgetPurpose, PurposeType::UniqueByName,
-		LOCTEXT("SceneOutlinerCellDescription", "The default widget to use in headers for the Scene Outliner."));
+		LOCTEXT("DefaultHeaderWidgetPurpose", "The default widget to use in headers for the Scene Outliner."));
 	
 	DataStorageUi.RegisterWidgetPurpose(FTypedElementSceneOutlinerQueryBinder::CellWidgetPurpose, PurposeType::UniqueByNameAndColumn,
-		LOCTEXT("SceneOutlinerCellDescription", "Widgets for cells in any Scene Outliner for specific columns or column combinations."));
+		LOCTEXT("CellWidgetPurpose", "Widgets for cells in any Scene Outliner for specific columns or column combinations."));
 	DataStorageUi.RegisterWidgetPurpose(FTypedElementSceneOutlinerQueryBinder::DefaultCellWidgetPurpose, PurposeType::UniqueByName,
-		LOCTEXT("SceneOutlinerCellDescription", "The default widget to use in cells for the Scene Outliner."));
+		LOCTEXT("DefaultCellWidgetPurpose", "The default widget to use in cells for the Scene Outliner."));
 }
 
 
