Initial support for TEDS' subquery support.

This is the first step in supporting subqueries from a query callback in TEDS. This version is limited to running queries that don't access data as these are therefor safe to directly run.

diff --git a/Engine/Plugins/Experimental/TypedElementsDataStorage/Source/TypedElementsDataStorage/Private/Processors/TypedElementProcessorAdaptors.cpp b/Engine/Plugins/Experimental/TypedElementsDataStorage/Source/TypedElementsDataStorage/Private/Processors/TypedElementProcessorAdaptors.cpp
--- a/Engine/Plugins/Experimental/TypedElementsDataStorage/Source/TypedElementsDataStorage/Private/Processors/TypedElementProcessorAdaptors.cpp
+++ b/Engine/Plugins/Experimental/TypedElementsDataStorage/Source/TypedElementsDataStorage/Private/Processors/TypedElementProcessorAdaptors.cpp
@@ -5,6 +5,7 @@
 #include <utility>
 #include "MassCommonTypes.h"
 #include "MassExecutionContext.h"
+#include "Queries/TypedElementExtendedQueryStore.h"
 #include "TypedElementDatabase.h"
 
 struct FMassContextForwarderShared
@@ -64,8 +65,9 @@ struct FMassContextForwarderShared
 
 struct FMassContextForwarder final : public ITypedElementDataStorageInterface::IQueryContext
 {
-	explicit FMassContextForwarder(FMassExecutionContext& Context)
-		: Context(Context)
+	FMassContextForwarder(FMassExecutionContext& InContext, FTypedElementExtendedQueryStore& InQueryStore)
+		: Context(InContext)
+		, QueryStore(InQueryStore)
 	{}
 
 	~FMassContextForwarder() override = default;
@@ -222,7 +224,16 @@ struct FMassContextForwarder final : public ITypedElementDataStorageInterface::I
 		}
 	}
 
+	ITypedElementDataStorageInterface::FQueryResult RunQuery(TypedElementQueryHandle Query) override
+	{
+		FTypedElementExtendedQueryStore::Handle Handle;
+		Handle.Handle = Query;
+		// This can be safely called because there's not callback, which means no columns are accessed, even for select queries.
+		return QueryStore.RunQuery(Context.GetEntityManagerChecked(), Handle);
+	}
+
 	FMassExecutionContext& Context;
+	FTypedElementExtendedQueryStore& QueryStore;
 };
 
 struct FMassDirectContextForwarder final : public ITypedElementDataStorageInterface::IDirectQueryContext
@@ -284,15 +295,18 @@ FPhasePreOrPostAmbleExecutor::~FPhasePreOrPostAmbleExecutor()
 	Context.FlushDeferred();
 }
 
-void FPhasePreOrPostAmbleExecutor::ExecuteQuery(ITypedElementDataStorageInterface::FQueryDescription& Description, FMassEntityQuery& NativeQuery,
+void FPhasePreOrPostAmbleExecutor::ExecuteQuery(
+	ITypedElementDataStorageInterface::FQueryDescription& Description, 
+	FTypedElementExtendedQueryStore& QueryStore, 
+	FMassEntityQuery& NativeQuery,
 	ITypedElementDataStorageInterface::QueryCallbackRef Callback)
 {
 	NativeQuery.ForEachEntityChunk(Context.GetEntityManagerChecked(), Context,
-		[&Callback, &Description](FMassExecutionContext& ExecutionContext)
+		[&Callback, &QueryStore, &Description](FMassExecutionContext& ExecutionContext)
 		{
 			if (FTypedElementQueryProcessorData::PrepareCachedDependenciesOnQuery(Description, ExecutionContext))
 			{
-				FMassContextForwarder QueryContext(ExecutionContext);
+				FMassContextForwarder QueryContext(ExecutionContext, QueryStore);
 				Callback(Description, QueryContext);
 			}
 		}
@@ -398,11 +412,11 @@ void FTypedElementQueryProcessorData::Execute(FMassEntityManager& EntityManager,
 	
 	ITypedElementDataStorageInterface::FQueryDescription& Description = ParentQuery->Description;
 	Query.ForEachEntityChunk(EntityManager, Context,
-		[&Description](FMassExecutionContext& Context)
+		[this, &Description](FMassExecutionContext& Context)
 		{
 			if (PrepareCachedDependenciesOnQuery(Description, Context))
 			{
-				FMassContextForwarder QueryContext(Context);
+				FMassContextForwarder QueryContext(Context, *QueryStore);
 				Description.Callback.Function(Description, QueryContext);
 			}
 		}
@@ -427,9 +441,11 @@ FMassEntityQuery& UTypedElementQueryProcessorCallbackAdapterProcessor::GetQuery(
 	return Data.Query;
 }
 
-void UTypedElementQueryProcessorCallbackAdapterProcessor::ConfigureQueryCallback(FTypedElementExtendedQuery& TargetParentQuery)
+void UTypedElementQueryProcessorCallbackAdapterProcessor::ConfigureQueryCallback(
+	FTypedElementExtendedQuery& TargetParentQuery, FTypedElementExtendedQueryStore& QueryStore)
 {
 	Data.ParentQuery = &TargetParentQuery;
+	Data.QueryStore = &QueryStore;
 
 	bRequiresGameThreadExecution = TargetParentQuery.Description.Callback.bForceToGameThread;
 	ExecutionFlags = static_cast<int32>(EProcessorExecutionFlags::Editor); 
@@ -491,9 +507,11 @@ EMassObservedOperation UTypedElementQueryObserverCallbackAdapterProcessor::GetOb
 	return Operation;
 }
 
-void UTypedElementQueryObserverCallbackAdapterProcessor::ConfigureQueryCallback(FTypedElementExtendedQuery& TargetParentQuery)
+void UTypedElementQueryObserverCallbackAdapterProcessor::ConfigureQueryCallback(
+	FTypedElementExtendedQuery& TargetParentQuery, FTypedElementExtendedQueryStore& QueryStore)
 {
 	Data.ParentQuery = &TargetParentQuery;
+	Data.QueryStore = &QueryStore;
 
 	bRequiresGameThreadExecution = TargetParentQuery.Description.Callback.bForceToGameThread;
 	ExecutionFlags = static_cast<int32>(EProcessorExecutionFlags::Editor);
diff --git a/Engine/Plugins/Experimental/TypedElementsDataStorage/Source/TypedElementsDataStorage/Private/Processors/TypedElementProcessorAdaptors.h b/Engine/Plugins/Experimental/TypedElementsDataStorage/Source/TypedElementsDataStorage/Private/Processors/TypedElementProcessorAdaptors.h
--- a/Engine/Plugins/Experimental/TypedElementsDataStorage/Source/TypedElementsDataStorage/Private/Processors/TypedElementProcessorAdaptors.h
+++ b/Engine/Plugins/Experimental/TypedElementsDataStorage/Source/TypedElementsDataStorage/Private/Processors/TypedElementProcessorAdaptors.h
@@ -12,13 +12,17 @@
 #include "TypedElementProcessorAdaptors.generated.h"
 
 struct FTypedElementExtendedQuery;
+class FTypedElementExtendedQueryStore;
 
 struct FPhasePreOrPostAmbleExecutor
 {
 	FPhasePreOrPostAmbleExecutor(FMassEntityManager& EntityManager, float DeltaTime);
 	~FPhasePreOrPostAmbleExecutor();
 
-	void ExecuteQuery(ITypedElementDataStorageInterface::FQueryDescription& Description, FMassEntityQuery& NativeQuery,
+	void ExecuteQuery(
+		ITypedElementDataStorageInterface::FQueryDescription& Description,
+		FTypedElementExtendedQueryStore& QueryStore,
+		FMassEntityQuery& NativeQuery,
 		ITypedElementDataStorageInterface::QueryCallbackRef Callback);
 
 	FMassExecutionContext Context;
@@ -46,6 +50,7 @@ struct FTypedElementQueryProcessorData
 		ITypedElementDataStorageInterface::FQueryDescription& Description, FMassExecutionContext& Context);
 
 	FTypedElementExtendedQuery* ParentQuery{ nullptr };
+	FTypedElementExtendedQueryStore* QueryStore{ nullptr };
 	FMassEntityQuery Query;
 };
 
@@ -61,7 +66,7 @@ public:
 	UTypedElementQueryProcessorCallbackAdapterProcessor();
 
 	FMassEntityQuery& GetQuery();
-	void ConfigureQueryCallback(FTypedElementExtendedQuery& Query);
+	void ConfigureQueryCallback(FTypedElementExtendedQuery& Query, FTypedElementExtendedQueryStore& QueryStore);
 
 protected:
 	void ConfigureQueries() override;
@@ -89,7 +94,7 @@ public:
 	FMassEntityQuery& GetQuery();
 	const UScriptStruct* GetObservedType() const;
 	EMassObservedOperation GetObservedOperation() const;
-	void ConfigureQueryCallback(FTypedElementExtendedQuery& Query);
+	void ConfigureQueryCallback(FTypedElementExtendedQuery& Query, FTypedElementExtendedQueryStore& QueryStore);
 
 protected:
 	void ConfigureQueries() override;
diff --git a/Engine/Plugins/Experimental/TypedElementsDataStorage/Source/TypedElementsDataStorage/Private/Queries/TypedElementExtendedQueryStore.cpp b/Engine/Plugins/Experimental/TypedElementsDataStorage/Source/TypedElementsDataStorage/Private/Queries/TypedElementExtendedQueryStore.cpp
--- a/Engine/Plugins/Experimental/TypedElementsDataStorage/Source/TypedElementsDataStorage/Private/Queries/TypedElementExtendedQueryStore.cpp
+++ b/Engine/Plugins/Experimental/TypedElementsDataStorage/Source/TypedElementsDataStorage/Private/Queries/TypedElementExtendedQueryStore.cpp
@@ -499,14 +499,14 @@ bool FTypedElementExtendedQueryStore::SetupProcessors(Handle Query, FTypedElemen
 		if (StoredQuery.Processor->IsA<UTypedElementQueryProcessorCallbackAdapterProcessor>())
 		{
 			static_cast<UTypedElementQueryProcessorCallbackAdapterProcessor*>(StoredQuery.Processor.Get())->
-				ConfigureQueryCallback(StoredQuery);
+				ConfigureQueryCallback(StoredQuery, *this);
 			PhaseManager.RegisterDynamicProcessor(*StoredQuery.Processor);
 		}
 		else if (StoredQuery.Processor->IsA<UTypedElementQueryObserverCallbackAdapterProcessor>())
 		{
 			UTypedElementQueryObserverCallbackAdapterProcessor* Observer =
 				static_cast<UTypedElementQueryObserverCallbackAdapterProcessor*>(StoredQuery.Processor.Get());
-			Observer->ConfigureQueryCallback(StoredQuery);
+			Observer->ConfigureQueryCallback(StoredQuery, *this);
 			EntityManager.GetObserverManager().AddObserverInstance(*Observer->GetObservedType(), Observer->GetObservedOperation(), *Observer);
 		}
 		else
@@ -570,7 +570,7 @@ void FTypedElementExtendedQueryStore::RunPhasePreOrPostAmbleQueries(FMassEntityM
 		for (Handle Query : QueryHandles)
 		{
 			FTypedElementExtendedQuery& QueryData = Queries.Get(Query);
-			Executor.ExecuteQuery(QueryData.Description, QueryData.NativeQuery, QueryData.Description.Callback.Function);
+			Executor.ExecuteQuery(QueryData.Description, *this, QueryData.NativeQuery, QueryData.Description.Callback.Function);
 		}
 	}
 }
diff --git a/Engine/Plugins/Experimental/TypedElementsDataStorage/Source/TypedElementsDataStorageUI/Private/TypedElementCounterWidget.cpp b/Engine/Plugins/Experimental/TypedElementsDataStorage/Source/TypedElementsDataStorageUI/Private/TypedElementCounterWidget.cpp
--- a/Engine/Plugins/Experimental/TypedElementsDataStorage/Source/TypedElementsDataStorageUI/Private/TypedElementCounterWidget.cpp
+++ b/Engine/Plugins/Experimental/TypedElementsDataStorage/Source/TypedElementsDataStorageUI/Private/TypedElementCounterWidget.cpp
@@ -12,7 +12,6 @@
 #include "Misc/CoreDelegates.h"
 #include "ToolMenu.h"
 #include "ToolMenus.h"
-#include "TypedElementSubsystems.h"
 #include "Widgets/DeclarativeSyntaxSupport.h"
 #include "Widgets/SWindow.h"
 #include "Widgets/Text/STextBlock.h"
@@ -40,18 +39,13 @@ void UTypedElementCounterWidgetFactory::RegisterQueries(ITypedElementDataStorage
 		FProcessor(DSI::EQueryTickPhase::FrameEnd, DataStorage.GetQueryTickGroupName(DSI::EQueryTickGroups::SyncWidgets))
 			.ForceToGameThread(true),
 		[](
-			FCachedQueryContext<UTypedElementDataStorageSubsystem>& Context, 
+			DSI::IQueryContext& Context,
 			FTypedElementSlateWidgetReferenceColumn& Widget,
 			FTypedElementU32IntValueCacheColumn& Comparison, 
 			const FTypedElementCounterWidgetColumn& Counter
 		)
 		{
-			UTypedElementDataStorageSubsystem& Subsystem = Context.GetCachedMutableDependency<UTypedElementDataStorageSubsystem>();
-			DSI* DataInterface = Subsystem.Get();
-			checkf(DataInterface, TEXT("FTypedElementsDataStorageUiModule tried to process widgets before the "
-				"Typed Elements Data Storage interface is available."));
-
-			DSI::FQueryResult Result = DataInterface->RunQuery(Counter.Query);
+			DSI::FQueryResult Result = Context.RunQuery(Counter.Query);
 			if (Result.Completed == DSI::FQueryResult::ECompletion::Fully && Result.Count != Comparison.Value)
 			{
 				TSharedPtr<SWidget> WidgetPointer = Widget.Widget.Pin();
diff --git a/Engine/Source/Runtime/TypedElementFramework/Public/Elements/Framework/TypedElementQueryBuilder.h b/Engine/Source/Runtime/TypedElementFramework/Public/Elements/Framework/TypedElementQueryBuilder.h
--- a/Engine/Source/Runtime/TypedElementFramework/Public/Elements/Framework/TypedElementQueryBuilder.h
+++ b/Engine/Source/Runtime/TypedElementFramework/Public/Elements/Framework/TypedElementQueryBuilder.h
@@ -250,6 +250,8 @@ namespace TypedElementQueryBuilder
 		inline void RemoveColumns(TypedElementRowHandle Row, TConstArrayView<const UScriptStruct*> ColumnTypes) override;
 		inline void RemoveColumns(TConstArrayView<TypedElementRowHandle> Rows, TConstArrayView<const UScriptStruct*> ColumnTypes) override;
 
+		inline ITypedElementDataStorageInterface::FQueryResult RunQuery(TypedElementQueryHandle Query) override;
+
 		ITypedElementDataStorageInterface::IQueryContext& ParentContext;
 		const ITypedElementDataStorageInterface::FQueryDescription& Description;
 	};
diff --git a/Engine/Source/Runtime/TypedElementFramework/Public/Elements/Framework/TypedElementQueryBuilder.inl b/Engine/Source/Runtime/TypedElementFramework/Public/Elements/Framework/TypedElementQueryBuilder.inl
--- a/Engine/Source/Runtime/TypedElementFramework/Public/Elements/Framework/TypedElementQueryBuilder.inl
+++ b/Engine/Source/Runtime/TypedElementFramework/Public/Elements/Framework/TypedElementQueryBuilder.inl
@@ -188,6 +188,11 @@ namespace TypedElementQueryBuilder
 		ParentContext.RemoveColumns(Rows, ColumnTypes);
 	}
 
+	ITypedElementDataStorageInterface::FQueryResult FQueryContextForwarder::RunQuery(TypedElementQueryHandle Query)
+	{
+		return ParentContext.RunQuery(Query);
+	}
+
 	//
 	// FCachedQueryContext
 	//
diff --git a/Engine/Source/Runtime/TypedElementFramework/Public/Elements/Interfaces/TypedElementDataStorageInterface.h b/Engine/Source/Runtime/TypedElementFramework/Public/Elements/Interfaces/TypedElementDataStorageInterface.h
--- a/Engine/Source/Runtime/TypedElementFramework/Public/Elements/Interfaces/TypedElementDataStorageInterface.h
+++ b/Engine/Source/Runtime/TypedElementFramework/Public/Elements/Interfaces/TypedElementDataStorageInterface.h
@@ -387,6 +387,12 @@ public:
 		 */
 		virtual void RemoveColumns(TConstArrayView<TypedElementRowHandle> Rows, TConstArrayView<const UScriptStruct*> ColumnTypes) = 0;
 
+		/**
+		 * Runs a previously created query. This version takes an arbitrary query, but is limited to running queries that do not directly 
+		 * access data from rows such as count queries.
+		 */
+		virtual FQueryResult RunQuery(TypedElementQueryHandle Query) = 0;
+
 		// Utility functions
 
 		template<typename... Columns>
