TEDS.PrintSupportedColumns now prints the structure size of the available columns.

#nojira

diff --git a/Engine/Plugins/Experimental/TypedElementsDataStorage/Source/TypedElementsDataStorage/Private/TypedElementDatabase.cpp b/Engine/Plugins/Experimental/TypedElementsDataStorage/Source/TypedElementsDataStorage/Private/TypedElementDatabase.cpp
--- a/Engine/Plugins/Experimental/TypedElementsDataStorage/Source/TypedElementsDataStorage/Private/TypedElementDatabase.cpp
+++ b/Engine/Plugins/Experimental/TypedElementsDataStorage/Source/TypedElementsDataStorage/Private/TypedElementDatabase.cpp
@@ -46,7 +46,19 @@ FAutoConsoleCommandWithOutputDevice PrintSupportedColumnsConsoleCommand(
 			{
 				if (It->IsChildOf(FragmentTypeInfo) || It->IsChildOf(TagTypeInfo))
 				{
-					Output.Logf(TEXT("    %s"), *It->GetFullName());
+					int32 StructureSize = It->GetStructureSize();
+					if (StructureSize >= 1024 * 1024)
+					{
+						Output.Logf(TEXT("    [%6.2f mib] %s"), StructureSize / (1024.0f * 1024.0f), *It->GetFullName());
+					}
+					else if (StructureSize >= 1024)
+					{
+						Output.Logf(TEXT("    [%6.2f kib] %s"), StructureSize / 1024.0f, *It->GetFullName());
+					}
+					else
+					{
+						Output.Logf(TEXT("    [%6i b  ] %s"), StructureSize, *It->GetFullName());
+					}
 				}
 			}
 			Output.Log(TEXT("End of Typed Elements Data Storage supported column list."));
