Allowed subqueries to be registered and called from another query in TEDS.

The following changes were made:
- Subqueries were added to the query description.
- The query builder can now register subqueries as a dependency.
- Subqueries can be run from a query callback. Currently subqueries can only be run without a callback and without any filtering on provided rows.
- The processor adaptors have been updated to take subqueries into account. Mass requires the native query to be on the processor for processor bound Mass queries. Due to this requirement there was a  need to have a small number of additional processors with a fixed number of native queries. 4 were set of, meaning that TEDS currently supports up to 4 subqueries per main query.
- Query builder calls that take ranges as input now reserve memory upfront.
- Compiling a query description from a query builder will now shrink the buffers and free any unused memory.
- The use of std::initializer_list in the query builder was updated to use TConstArrayView as these are a better fit.

diff --git a/Engine/Plugins/Experimental/TypedElementsDataStorage/Source/TypedElementsDataStorage/Private/Processors/TypedElementProcessorAdaptors.cpp b/Engine/Plugins/Experimental/TypedElementsDataStorage/Source/TypedElementsDataStorage/Private/Processors/TypedElementProcessorAdaptors.cpp
--- a/Engine/Plugins/Experimental/TypedElementsDataStorage/Source/TypedElementsDataStorage/Private/Processors/TypedElementProcessorAdaptors.cpp
+++ b/Engine/Plugins/Experimental/TypedElementsDataStorage/Source/TypedElementsDataStorage/Private/Processors/TypedElementProcessorAdaptors.cpp
@@ -65,8 +65,10 @@ struct FMassContextForwarderShared
 
 struct FMassContextForwarder final : public ITypedElementDataStorageInterface::IQueryContext
 {
-	FMassContextForwarder(FMassExecutionContext& InContext, FTypedElementExtendedQueryStore& InQueryStore)
-		: Context(InContext)
+	FMassContextForwarder(ITypedElementDataStorageInterface::FQueryDescription& InQueryDescription, FMassExecutionContext& InContext, 
+		FTypedElementExtendedQueryStore& InQueryStore)
+		: QueryDescription(InQueryDescription)
+		, Context(InContext)
 		, QueryStore(InQueryStore)
 	{}
 
@@ -232,6 +234,14 @@ struct FMassContextForwarder final : public ITypedElementDataStorageInterface::I
 		return QueryStore.RunQuery(Context.GetEntityManagerChecked(), Handle);
 	}
 
+	ITypedElementDataStorageInterface::FQueryResult RunSubquery(int32 SubqueryIndex) override
+	{
+		return SubqueryIndex < QueryDescription.Subqueries.Num() ?
+			RunQuery(QueryDescription.Subqueries[SubqueryIndex]) :
+			ITypedElementDataStorageInterface::FQueryResult{};
+	}
+
+	ITypedElementDataStorageInterface::FQueryDescription& QueryDescription;
 	FMassExecutionContext& Context;
 	FTypedElementExtendedQueryStore& QueryStore;
 };
@@ -306,7 +316,7 @@ void FPhasePreOrPostAmbleExecutor::ExecuteQuery(
 		{
 			if (FTypedElementQueryProcessorData::PrepareCachedDependenciesOnQuery(Description, ExecutionContext))
 			{
-				FMassContextForwarder QueryContext(ExecutionContext, QueryStore);
+				FMassContextForwarder QueryContext(Description, ExecutionContext, QueryStore);
 				Callback(Description, QueryContext);
 			}
 		}
@@ -319,10 +329,49 @@ void FPhasePreOrPostAmbleExecutor::ExecuteQuery(
  * FTypedElementQueryProcessorData
  */
 FTypedElementQueryProcessorData::FTypedElementQueryProcessorData(UMassProcessor& Owner)
-	: Query(Owner)
+	: NativeQuery(Owner)
 {
 }
 
+bool FTypedElementQueryProcessorData::CommonQueryConfiguration(UMassProcessor& InOwner, FTypedElementExtendedQuery& InQuery, 
+	FTypedElementExtendedQueryStore& InQueryStore, TArrayView<FMassEntityQuery> Subqueries)
+{
+	ParentQuery = &InQuery;
+	QueryStore = &InQueryStore;
+
+	if (ensureMsgf(InQuery.Description.Subqueries.Num() <= Subqueries.Num(),
+		TEXT("Provided query has too many (%i) subqueries."), InQuery.Description.Subqueries.Num()))
+	{
+		bool Result = true;
+		int32 CurrentSubqueryIndex = 0;
+		for (TypedElementQueryHandle SubqueryHandle : InQuery.Description.Subqueries)
+		{
+			FTypedElementExtendedQueryStore::Handle SubqueryStoreHandle;
+			SubqueryStoreHandle.Handle = SubqueryHandle;
+			if (const FTypedElementExtendedQuery* Subquery = InQueryStore.Get(SubqueryStoreHandle))
+			{
+				if (ensureMsgf(Subquery->NativeQuery.CheckValidity(), TEXT("Provided subquery isn't valid. This can be because it couldn't be "
+					"constructed properly or because it's been bound to a callback.")))
+				{
+					Subqueries[CurrentSubqueryIndex] = Subquery->NativeQuery;
+					Subqueries[CurrentSubqueryIndex].RegisterWithProcessor(InOwner);
+					++CurrentSubqueryIndex;
+				}
+				else
+				{
+					Result = false;
+				}
+			}
+			else
+			{
+				Result = false;
+			}
+		}
+		return Result;
+	}
+	return false;
+}
+
 EMassProcessingPhase FTypedElementQueryProcessorData::MapToMassProcessingPhase(ITypedElementDataStorageInterface::EQueryTickPhase Phase)
 {
 	switch(Phase)
@@ -411,12 +460,12 @@ void FTypedElementQueryProcessorData::Execute(FMassEntityManager& EntityManager,
 	checkf(ParentQuery, TEXT("A query callback was registered for execution without an associated query."));
 	
 	ITypedElementDataStorageInterface::FQueryDescription& Description = ParentQuery->Description;
-	Query.ForEachEntityChunk(EntityManager, Context,
+	NativeQuery.ForEachEntityChunk(EntityManager, Context,
 		[this, &Description](FMassExecutionContext& Context)
 		{
 			if (PrepareCachedDependenciesOnQuery(Description, Context))
 			{
-				FMassContextForwarder QueryContext(Context, *QueryStore);
+				FMassContextForwarder QueryContext(Description, Context, *QueryStore);
 				Description.Callback.Function(Description, QueryContext);
 			}
 		}
@@ -429,96 +478,130 @@ void FTypedElementQueryProcessorData::Execute(FMassEntityManager& EntityManager,
  * UTypedElementQueryProcessorCallbackAdapterProcessor
  */
 
-UTypedElementQueryProcessorCallbackAdapterProcessor::UTypedElementQueryProcessorCallbackAdapterProcessor()
+UTypedElementQueryProcessorCallbackAdapterProcessorBase::UTypedElementQueryProcessorCallbackAdapterProcessorBase()
 	: Data(*this)
 {
 	bAllowMultipleInstances = true;
 	bAutoRegisterWithProcessingPhases = false;
 }
 
-FMassEntityQuery& UTypedElementQueryProcessorCallbackAdapterProcessor::GetQuery()
+FMassEntityQuery& UTypedElementQueryProcessorCallbackAdapterProcessorBase::GetQuery()
+{
+	return Data.NativeQuery;
+}
+
+bool UTypedElementQueryProcessorCallbackAdapterProcessorBase::ConfigureQueryCallback(
+	FTypedElementExtendedQuery& Query, FTypedElementExtendedQueryStore& QueryStore)
 {
-	return Data.Query;
+	return ConfigureQueryCallbackData(Query, QueryStore, {});
 }
 
-void UTypedElementQueryProcessorCallbackAdapterProcessor::ConfigureQueryCallback(
-	FTypedElementExtendedQuery& TargetParentQuery, FTypedElementExtendedQueryStore& QueryStore)
+bool UTypedElementQueryProcessorCallbackAdapterProcessorBase::ConfigureQueryCallbackData(FTypedElementExtendedQuery& Query,
+	FTypedElementExtendedQueryStore& QueryStore,  TArrayView<FMassEntityQuery> Subqueries)
 {
-	Data.ParentQuery = &TargetParentQuery;
-	Data.QueryStore = &QueryStore;
+	bool Result = Data.CommonQueryConfiguration(*this, Query, QueryStore, Subqueries);
 
-	bRequiresGameThreadExecution = TargetParentQuery.Description.Callback.bForceToGameThread;
+	bRequiresGameThreadExecution = Query.Description.Callback.bForceToGameThread;
 	ExecutionFlags = static_cast<int32>(EProcessorExecutionFlags::Editor); 
-	ExecutionOrder.ExecuteInGroup = TargetParentQuery.Description.Callback.Group;
-	ExecutionOrder.ExecuteBefore = TargetParentQuery.Description.Callback.BeforeGroups;
-	ExecutionOrder.ExecuteAfter = TargetParentQuery.Description.Callback.AfterGroups;
-	ProcessingPhase = Data.MapToMassProcessingPhase(TargetParentQuery.Description.Callback.Phase);
+	ExecutionOrder.ExecuteInGroup = Query.Description.Callback.Group;
+	ExecutionOrder.ExecuteBefore = Query.Description.Callback.BeforeGroups;
+	ExecutionOrder.ExecuteAfter = Query.Description.Callback.AfterGroups;
+	ProcessingPhase = Data.MapToMassProcessingPhase(Query.Description.Callback.Phase);
 
 	Super::PostInitProperties();
+	return Result;
 }
 
-void UTypedElementQueryProcessorCallbackAdapterProcessor::ConfigureQueries()
+void UTypedElementQueryProcessorCallbackAdapterProcessorBase::ConfigureQueries()
 {
 	// When the extended query information is provided the native query will already be fully configured.
 }
 
-void UTypedElementQueryProcessorCallbackAdapterProcessor::PostInitProperties()
+void UTypedElementQueryProcessorCallbackAdapterProcessorBase::PostInitProperties()
 {
 	Super::Super::PostInitProperties();
 }
 
-FString UTypedElementQueryProcessorCallbackAdapterProcessor::GetProcessorName() const
+FString UTypedElementQueryProcessorCallbackAdapterProcessorBase::GetProcessorName() const
 {
 	FString Name = Data.GetProcessorName();
 	Name += TEXT(" [Editor Processor]");
 	return Name;
 }
 
-void UTypedElementQueryProcessorCallbackAdapterProcessor::Execute(FMassEntityManager& EntityManager, FMassExecutionContext& Context)
+void UTypedElementQueryProcessorCallbackAdapterProcessorBase::Execute(FMassEntityManager& EntityManager, FMassExecutionContext& Context)
 {
 	Data.Execute(EntityManager, Context);
 }
 
+bool UTypedElementQueryProcessorCallbackAdapterProcessorWith1Subquery::ConfigureQueryCallback(
+	FTypedElementExtendedQuery& Query, FTypedElementExtendedQueryStore& QueryStore)
+{
+	return ConfigureQueryCallbackData(Query, QueryStore, NativeSubqueries);
+}
+
+bool UTypedElementQueryProcessorCallbackAdapterProcessorWith2Subqueries::ConfigureQueryCallback(
+	FTypedElementExtendedQuery& Query, FTypedElementExtendedQueryStore& QueryStore)
+{
+	return ConfigureQueryCallbackData(Query, QueryStore, NativeSubqueries);
+}
+
+bool UTypedElementQueryProcessorCallbackAdapterProcessorWith3Subqueries::ConfigureQueryCallback(
+	FTypedElementExtendedQuery& Query, FTypedElementExtendedQueryStore& QueryStore)
+{
+	return ConfigureQueryCallbackData(Query, QueryStore, NativeSubqueries);
+}
+
+bool UTypedElementQueryProcessorCallbackAdapterProcessorWith4Subqueries::ConfigureQueryCallback(
+	FTypedElementExtendedQuery& Query, FTypedElementExtendedQueryStore& QueryStore)
+{
+	return ConfigureQueryCallbackData(Query, QueryStore, NativeSubqueries);
+}
 
 
 /**
  * UTypedElementQueryObserverCallbackAdapterProcessor
  */
 
-UTypedElementQueryObserverCallbackAdapterProcessor::UTypedElementQueryObserverCallbackAdapterProcessor()
+UTypedElementQueryObserverCallbackAdapterProcessorBase::UTypedElementQueryObserverCallbackAdapterProcessorBase()
 	: Data(*this)
 {
 	bAllowMultipleInstances = true;
 	bAutoRegisterWithProcessingPhases = false;
 }
 
-FMassEntityQuery& UTypedElementQueryObserverCallbackAdapterProcessor::GetQuery()
+FMassEntityQuery& UTypedElementQueryObserverCallbackAdapterProcessorBase::GetQuery()
 {
-	return Data.Query;
+	return Data.NativeQuery;
 }
 
-const UScriptStruct* UTypedElementQueryObserverCallbackAdapterProcessor::GetObservedType() const
+const UScriptStruct* UTypedElementQueryObserverCallbackAdapterProcessorBase::GetObservedType() const
 {
 	return ObservedType;
 }
 
-EMassObservedOperation UTypedElementQueryObserverCallbackAdapterProcessor::GetObservedOperation() const
+EMassObservedOperation UTypedElementQueryObserverCallbackAdapterProcessorBase::GetObservedOperation() const
 {
 	return Operation;
 }
 
-void UTypedElementQueryObserverCallbackAdapterProcessor::ConfigureQueryCallback(
-	FTypedElementExtendedQuery& TargetParentQuery, FTypedElementExtendedQueryStore& QueryStore)
+bool UTypedElementQueryObserverCallbackAdapterProcessorBase::ConfigureQueryCallback(
+	FTypedElementExtendedQuery& Query, FTypedElementExtendedQueryStore& QueryStore)
+{
+	return ConfigureQueryCallbackData(Query, QueryStore, {});
+}
+
+bool UTypedElementQueryObserverCallbackAdapterProcessorBase::ConfigureQueryCallbackData(FTypedElementExtendedQuery& Query,
+	FTypedElementExtendedQueryStore& QueryStore, TArrayView<FMassEntityQuery> Subqueries)
 {
-	Data.ParentQuery = &TargetParentQuery;
-	Data.QueryStore = &QueryStore;
+	bool Result = Data.CommonQueryConfiguration(*this, Query, QueryStore, Subqueries);
 
-	bRequiresGameThreadExecution = TargetParentQuery.Description.Callback.bForceToGameThread;
+	bRequiresGameThreadExecution = Query.Description.Callback.bForceToGameThread;
 	ExecutionFlags = static_cast<int32>(EProcessorExecutionFlags::Editor);
 	
-	ObservedType = const_cast<UScriptStruct*>(TargetParentQuery.Description.Callback.MonitoredType);
+	ObservedType = const_cast<UScriptStruct*>(Query.Description.Callback.MonitoredType);
 	
-	switch (TargetParentQuery.Description.Callback.Type)
+	switch (Query.Description.Callback.Type)
 	{
 	case ITypedElementDataStorageInterface::EQueryCallbackType::ObserveAdd:
 		Operation = EMassObservedOperation::Add;
@@ -528,29 +611,30 @@ void UTypedElementQueryObserverCallbackAdapterProcessor::ConfigureQueryCallback(
 		break;
 	default:
 		checkf(false, TEXT("Query type %i is not supported from the observer processor adapter."),
-			static_cast<int>(TargetParentQuery.Description.Callback.Type));
-		break;
+			static_cast<int>(Query.Description.Callback.Type));
+		return false;
 	}
 
 	Super::PostInitProperties();
+	return Result;
 }
 
-void UTypedElementQueryObserverCallbackAdapterProcessor::ConfigureQueries()
+void UTypedElementQueryObserverCallbackAdapterProcessorBase::ConfigureQueries()
 {
 	// When the extended query information is provided the native query will already be fully configured.
 }
 
-void UTypedElementQueryObserverCallbackAdapterProcessor::PostInitProperties()
+void UTypedElementQueryObserverCallbackAdapterProcessorBase::PostInitProperties()
 {
 	Super::Super::PostInitProperties();
 }
 
-void UTypedElementQueryObserverCallbackAdapterProcessor::Register()
+void UTypedElementQueryObserverCallbackAdapterProcessorBase::Register()
 { 
 	// Do nothing as this processor will be explicitly registered.
 }
 
-FString UTypedElementQueryObserverCallbackAdapterProcessor::GetProcessorName() const
+FString UTypedElementQueryObserverCallbackAdapterProcessorBase::GetProcessorName() const
 {
 	FString Name = Data.GetProcessorName();
 	EMassObservedOperation ObservationType = GetObservedOperation();
@@ -569,7 +653,31 @@ FString UTypedElementQueryObserverCallbackAdapterProcessor::GetProcessorName() c
 	return Name;
 }
 
-void UTypedElementQueryObserverCallbackAdapterProcessor::Execute(FMassEntityManager& EntityManager, FMassExecutionContext& Context)
+void UTypedElementQueryObserverCallbackAdapterProcessorBase::Execute(FMassEntityManager& EntityManager, FMassExecutionContext& Context)
 {
 	Data.Execute(EntityManager, Context);
 }
+
+bool UTypedElementQueryObserverCallbackAdapterProcessorWith1Subquery::ConfigureQueryCallback(
+	FTypedElementExtendedQuery& Query, FTypedElementExtendedQueryStore& QueryStore)
+{
+	return ConfigureQueryCallbackData(Query, QueryStore, NativeSubqueries);
+}
+
+bool UTypedElementQueryObserverCallbackAdapterProcessorWith2Subqueries::ConfigureQueryCallback(
+	FTypedElementExtendedQuery& Query, FTypedElementExtendedQueryStore& QueryStore)
+{
+	return ConfigureQueryCallbackData(Query, QueryStore, NativeSubqueries);
+}
+
+bool UTypedElementQueryObserverCallbackAdapterProcessorWith3Subqueries::ConfigureQueryCallback(
+	FTypedElementExtendedQuery& Query, FTypedElementExtendedQueryStore& QueryStore)
+{
+	return ConfigureQueryCallbackData(Query, QueryStore, NativeSubqueries);
+}
+
+bool UTypedElementQueryObserverCallbackAdapterProcessorWith4Subqueries::ConfigureQueryCallback(
+	FTypedElementExtendedQuery& Query, FTypedElementExtendedQueryStore& QueryStore)
+{
+	return ConfigureQueryCallbackData(Query, QueryStore, NativeSubqueries);
+}
diff --git a/Engine/Plugins/Experimental/TypedElementsDataStorage/Source/TypedElementsDataStorage/Private/Processors/TypedElementProcessorAdaptors.h b/Engine/Plugins/Experimental/TypedElementsDataStorage/Source/TypedElementsDataStorage/Private/Processors/TypedElementProcessorAdaptors.h
--- a/Engine/Plugins/Experimental/TypedElementsDataStorage/Source/TypedElementsDataStorage/Private/Processors/TypedElementProcessorAdaptors.h
+++ b/Engine/Plugins/Experimental/TypedElementsDataStorage/Source/TypedElementsDataStorage/Private/Processors/TypedElementProcessorAdaptors.h
@@ -36,6 +36,8 @@ struct FTypedElementQueryProcessorData
 	FTypedElementQueryProcessorData() = default;
 	explicit FTypedElementQueryProcessorData(UMassProcessor& Owner);
 
+	bool CommonQueryConfiguration(UMassProcessor& InOwner, FTypedElementExtendedQuery& InQuery, 
+		FTypedElementExtendedQueryStore& InQueryStore, TArrayView<FMassEntityQuery> Subqueries);
 	static EMassProcessingPhase MapToMassProcessingPhase(ITypedElementDataStorageInterface::EQueryTickPhase Phase);
 	FString GetProcessorName() const;
 
@@ -49,26 +51,29 @@ struct FTypedElementQueryProcessorData
 	static bool PrepareCachedDependenciesOnQuery(
 		ITypedElementDataStorageInterface::FQueryDescription& Description, FMassExecutionContext& Context);
 
+	
 	FTypedElementExtendedQuery* ParentQuery{ nullptr };
 	FTypedElementExtendedQueryStore* QueryStore{ nullptr };
-	FMassEntityQuery Query;
+	FMassEntityQuery NativeQuery;
 };
 
 /**
  * Adapts processor queries callback for MASS.
  */
 UCLASS()
-class UTypedElementQueryProcessorCallbackAdapterProcessor : public UMassProcessor
+class UTypedElementQueryProcessorCallbackAdapterProcessorBase : public UMassProcessor
 {
 	GENERATED_BODY()
 
 public:
-	UTypedElementQueryProcessorCallbackAdapterProcessor();
+	UTypedElementQueryProcessorCallbackAdapterProcessorBase();
 
 	FMassEntityQuery& GetQuery();
-	void ConfigureQueryCallback(FTypedElementExtendedQuery& Query, FTypedElementExtendedQueryStore& QueryStore);
+	virtual bool ConfigureQueryCallback(FTypedElementExtendedQuery& Query, FTypedElementExtendedQueryStore& QueryStore);
 
 protected:
+	bool ConfigureQueryCallbackData(FTypedElementExtendedQuery& Query, FTypedElementExtendedQueryStore& QueryStore, 
+		TArrayView<FMassEntityQuery> Subqueries);
 	void ConfigureQueries() override;
 	void Execute(FMassEntityManager& EntityManager, FMassExecutionContext& TargetParentQuery) override;
 	
@@ -80,23 +85,91 @@ private:
 	FTypedElementQueryProcessorData Data;
 };
 
+UCLASS()
+class UTypedElementQueryProcessorCallbackAdapterProcessor final : public UTypedElementQueryProcessorCallbackAdapterProcessorBase
+{
+	GENERATED_BODY()
+};
+
+/**
+ * Mass verifies that queries that are used by processors are on the processor themselves. It does this by taking the address of the query 
+ * and seeing if it's within the start and end address of the processor. When a dynamic array is used those addresses are going to be 
+ * elsewhere, so the two options are to store a single fixed size array on a processor or have multiple instances. With Mass' queries being 
+ * not an insignificant size it's preferable to have several variants with queries to allow the choice for the minimal size. Unfortunately 
+ * UHT doesn't allow for templates so it had to be done in an explicit way.
+ */
+
+UCLASS()
+class UTypedElementQueryProcessorCallbackAdapterProcessorWith1Subquery final : public UTypedElementQueryProcessorCallbackAdapterProcessorBase
+{
+	GENERATED_BODY()
+
+public:
+	bool ConfigureQueryCallback(FTypedElementExtendedQuery& Query, FTypedElementExtendedQueryStore& QueryStore) override;
+
+private:
+	UPROPERTY(transient)
+	FMassEntityQuery NativeSubqueries[1];
+};
+
+UCLASS()
+class UTypedElementQueryProcessorCallbackAdapterProcessorWith2Subqueries final : public UTypedElementQueryProcessorCallbackAdapterProcessorBase
+{
+	GENERATED_BODY()
+
+public:
+	bool ConfigureQueryCallback(FTypedElementExtendedQuery& Query, FTypedElementExtendedQueryStore& QueryStore) override;
+
+private:
+	UPROPERTY(transient)
+	FMassEntityQuery NativeSubqueries[2];
+};
+
+UCLASS()
+class UTypedElementQueryProcessorCallbackAdapterProcessorWith3Subqueries final : public UTypedElementQueryProcessorCallbackAdapterProcessorBase
+{
+	GENERATED_BODY()
+
+public:
+	bool ConfigureQueryCallback(FTypedElementExtendedQuery& Query, FTypedElementExtendedQueryStore& QueryStore) override;
+
+private:
+	UPROPERTY(transient)
+	FMassEntityQuery NativeSubqueries[3];
+};
+
+UCLASS()
+class UTypedElementQueryProcessorCallbackAdapterProcessorWith4Subqueries final : public UTypedElementQueryProcessorCallbackAdapterProcessorBase
+{
+	GENERATED_BODY()
+
+public:
+	bool ConfigureQueryCallback(FTypedElementExtendedQuery& Query, FTypedElementExtendedQueryStore& QueryStore) override;
+
+private:
+	UPROPERTY(transient)
+	FMassEntityQuery NativeSubqueries[4];
+};
+
 /**
  * Adapts observer queries callback for MASS.
  */
 UCLASS()
-class UTypedElementQueryObserverCallbackAdapterProcessor : public UMassObserverProcessor
+class UTypedElementQueryObserverCallbackAdapterProcessorBase : public UMassObserverProcessor
 {
 	GENERATED_BODY()
 
 public:
-	UTypedElementQueryObserverCallbackAdapterProcessor();
+	UTypedElementQueryObserverCallbackAdapterProcessorBase();
 
 	FMassEntityQuery& GetQuery();
 	const UScriptStruct* GetObservedType() const;
 	EMassObservedOperation GetObservedOperation() const;
-	void ConfigureQueryCallback(FTypedElementExtendedQuery& Query, FTypedElementExtendedQueryStore& QueryStore);
+	virtual bool ConfigureQueryCallback(FTypedElementExtendedQuery& Query, FTypedElementExtendedQueryStore& QueryStore);
 
 protected:
+	bool ConfigureQueryCallbackData(FTypedElementExtendedQuery& Query, FTypedElementExtendedQueryStore& QueryStore,
+		TArrayView<FMassEntityQuery> Subqueries);
 	void ConfigureQueries() override;
 	void Execute(FMassEntityManager& EntityManager, FMassExecutionContext& TargetParentQuery) override;
 
@@ -108,3 +181,69 @@ private:
 	UPROPERTY(transient)
 	FTypedElementQueryProcessorData Data;
 };
+
+UCLASS()
+class UTypedElementQueryObserverCallbackAdapterProcessor final : public UTypedElementQueryObserverCallbackAdapterProcessorBase
+{
+	GENERATED_BODY()
+};
+
+/**
+ * Mass verifies that queries that are used by processors are on the processor themselves. It does this by taking the address of the query
+ * and seeing if it's within the start and end address of the processor. When a dynamic array is used those addresses are going to be
+ * elsewhere, so the two options are to store a single fixed size array on a processor or have multiple instances. With Mass' queries being
+ * not an insignificant size it's preferable to have several variants with queries to allow the choice for the minimal size. Unfortunately
+ * UHT doesn't allow for templates so it had to be done in an explicit way.
+ */
+
+UCLASS()
+class UTypedElementQueryObserverCallbackAdapterProcessorWith1Subquery final : public UTypedElementQueryObserverCallbackAdapterProcessorBase
+{
+	GENERATED_BODY()
+
+public:
+	bool ConfigureQueryCallback(FTypedElementExtendedQuery& Query, FTypedElementExtendedQueryStore& QueryStore) override;
+
+private:
+	UPROPERTY(transient)
+	FMassEntityQuery NativeSubqueries[1];
+};
+
+UCLASS()
+class UTypedElementQueryObserverCallbackAdapterProcessorWith2Subqueries final : public UTypedElementQueryObserverCallbackAdapterProcessorBase
+{
+	GENERATED_BODY()
+
+public:
+	bool ConfigureQueryCallback(FTypedElementExtendedQuery& Query, FTypedElementExtendedQueryStore& QueryStore) override;
+
+private:
+	UPROPERTY(transient)
+	FMassEntityQuery NativeSubqueries[2];
+};
+
+UCLASS()
+class UTypedElementQueryObserverCallbackAdapterProcessorWith3Subqueries final : public UTypedElementQueryObserverCallbackAdapterProcessorBase
+{
+	GENERATED_BODY()
+
+public:
+	bool ConfigureQueryCallback(FTypedElementExtendedQuery& Query, FTypedElementExtendedQueryStore& QueryStore) override;
+
+private:
+	UPROPERTY(transient)
+	FMassEntityQuery NativeSubqueries[3];
+};
+
+UCLASS()
+class UTypedElementQueryObserverCallbackAdapterProcessorWith4Subqueries final : public UTypedElementQueryObserverCallbackAdapterProcessorBase
+{
+	GENERATED_BODY()
+
+public:
+	bool ConfigureQueryCallback(FTypedElementExtendedQuery& Query, FTypedElementExtendedQueryStore& QueryStore) override;
+
+private:
+	UPROPERTY(transient)
+	FMassEntityQuery NativeSubqueries[4];
+};
diff --git a/Engine/Plugins/Experimental/TypedElementsDataStorage/Source/TypedElementsDataStorage/Private/Queries/TypedElementExtendedQueryStore.cpp b/Engine/Plugins/Experimental/TypedElementsDataStorage/Source/TypedElementsDataStorage/Private/Queries/TypedElementExtendedQueryStore.cpp
--- a/Engine/Plugins/Experimental/TypedElementsDataStorage/Source/TypedElementsDataStorage/Private/Queries/TypedElementExtendedQueryStore.cpp
+++ b/Engine/Plugins/Experimental/TypedElementsDataStorage/Source/TypedElementsDataStorage/Private/Queries/TypedElementExtendedQueryStore.cpp
@@ -38,11 +38,11 @@ void FTypedElementExtendedQueryStore::UnregisterQuery(Handle Query, FMassProcess
 	{
 		if (QueryData->Processor)
 		{
-			if (QueryData->Processor->IsA<UTypedElementQueryProcessorCallbackAdapterProcessor>())
+			if (QueryData->Processor->IsA<UTypedElementQueryProcessorCallbackAdapterProcessorBase>())
 			{
 				PhaseManager.UnregisterDynamicProcessor(*QueryData->Processor);
 			}
-			else if (QueryData->Processor->IsA<UTypedElementQueryObserverCallbackAdapterProcessor>())
+			else if (QueryData->Processor->IsA<UTypedElementQueryObserverCallbackAdapterProcessorBase>())
 			{
 				checkf(false, TEXT("Observer queries can not be unregistered."));
 			}
@@ -277,6 +277,14 @@ void FTypedElementExtendedQueryStore::DebugPrintQueryCallbacks(FOutputDevice& Ou
 FMassEntityQuery& FTypedElementExtendedQueryStore::SetupNativeQuery(
 	ITypedElementDataStorageInterface::FQueryDescription& Query, FTypedElementExtendedQuery& StoredQuery)
 {
+	/**
+	 * Mass verifies that queries that are used by processors are on the processor themselves. It does this by taking the address of the query
+	 * and seeing if it's within the start and end address of the processor. When a dynamic array is used those addresses are going to be
+	 * elsewhere, so the two options are to store a single fixed size array on a processor or have multiple instances. With Mass' queries being
+	 * not an insignificant size it's preferable to have several variants with queries to allow the choice for the minimal size. Unfortunately
+	 * UHT doesn't allow for templates so it had to be done in an explicit way.
+	 */
+
 	using DSI = ITypedElementDataStorageInterface;
 
 	if (Query.Action == DSI::FQueryDescription::EActionType::Select)
@@ -287,8 +295,29 @@ FMassEntityQuery& FTypedElementExtendedQueryStore::SetupNativeQuery(
 			break;
 		case DSI::EQueryCallbackType::Processor:
 		{
-			UTypedElementQueryProcessorCallbackAdapterProcessor* Processor =
-				NewObject<UTypedElementQueryProcessorCallbackAdapterProcessor>();
+			UTypedElementQueryProcessorCallbackAdapterProcessorBase* Processor;
+			switch (Query.Subqueries.Num())
+			{
+			case 0:
+				Processor = NewObject<UTypedElementQueryProcessorCallbackAdapterProcessor>();
+				break;
+			case 1:
+				Processor = NewObject<UTypedElementQueryProcessorCallbackAdapterProcessorWith1Subquery>();
+				break;
+			case 2:
+				Processor = NewObject<UTypedElementQueryProcessorCallbackAdapterProcessorWith2Subqueries>();
+				break;
+			case 3:
+				Processor = NewObject<UTypedElementQueryProcessorCallbackAdapterProcessorWith3Subqueries>();
+				break;
+			case 4:
+				Processor = NewObject<UTypedElementQueryProcessorCallbackAdapterProcessorWith4Subqueries>();
+				break;
+			default:
+				checkf(false, TEXT("The current Typed Elements Data Storage backend doesn't support %i subqueries per processor query."), 
+					Query.Subqueries.Num());
+				return StoredQuery.NativeQuery;
+			}
 			StoredQuery.Processor.Reset(Processor);
 			return Processor->GetQuery();
 		}
@@ -296,8 +325,29 @@ FMassEntityQuery& FTypedElementExtendedQueryStore::SetupNativeQuery(
 			// Fall-through
 		case DSI::EQueryCallbackType::ObserveRemove:
 		{
-			UTypedElementQueryObserverCallbackAdapterProcessor* Observer =
-				NewObject<UTypedElementQueryObserverCallbackAdapterProcessor>();
+			UTypedElementQueryObserverCallbackAdapterProcessorBase* Observer;
+			switch (Query.Subqueries.Num())
+			{
+			case 0:
+				Observer = NewObject<UTypedElementQueryObserverCallbackAdapterProcessor>();
+				break;
+			case 1:
+				Observer = NewObject<UTypedElementQueryObserverCallbackAdapterProcessorWith1Subquery>();
+				break;
+			case 2:
+				Observer = NewObject<UTypedElementQueryObserverCallbackAdapterProcessorWith2Subqueries>();
+				break;
+			case 3:
+				Observer = NewObject<UTypedElementQueryObserverCallbackAdapterProcessorWith3Subqueries>();
+				break;
+			case 4:
+				Observer = NewObject<UTypedElementQueryObserverCallbackAdapterProcessorWith4Subqueries>();
+				break;
+			default:
+				checkf(false, TEXT("The current Typed Elements Data Storage backend doesn't support %i subqueries per observer query."),
+					Query.Subqueries.Num());
+				return StoredQuery.NativeQuery;
+			}
 			StoredQuery.Processor.Reset(Observer);
 			return Observer->GetQuery();
 		}
@@ -496,18 +546,30 @@ bool FTypedElementExtendedQueryStore::SetupProcessors(Handle Query, FTypedElemen
 	// Register regular processors and observer with Mass.
 	if (StoredQuery.Processor)
 	{
-		if (StoredQuery.Processor->IsA<UTypedElementQueryProcessorCallbackAdapterProcessor>())
+		if (StoredQuery.Processor->IsA<UTypedElementQueryProcessorCallbackAdapterProcessorBase>())
 		{
-			static_cast<UTypedElementQueryProcessorCallbackAdapterProcessor*>(StoredQuery.Processor.Get())->
-				ConfigureQueryCallback(StoredQuery, *this);
-			PhaseManager.RegisterDynamicProcessor(*StoredQuery.Processor);
+			if (static_cast<UTypedElementQueryProcessorCallbackAdapterProcessorBase*>(StoredQuery.Processor.Get())->
+				ConfigureQueryCallback(StoredQuery, *this))
+			{
+				PhaseManager.RegisterDynamicProcessor(*StoredQuery.Processor);
+			}
+			else
+			{
+				return false;
+			}
 		}
-		else if (StoredQuery.Processor->IsA<UTypedElementQueryObserverCallbackAdapterProcessor>())
+		else if (StoredQuery.Processor->IsA<UTypedElementQueryObserverCallbackAdapterProcessorBase>())
 		{
-			UTypedElementQueryObserverCallbackAdapterProcessor* Observer =
-				static_cast<UTypedElementQueryObserverCallbackAdapterProcessor*>(StoredQuery.Processor.Get());
-			Observer->ConfigureQueryCallback(StoredQuery, *this);
-			EntityManager.GetObserverManager().AddObserverInstance(*Observer->GetObservedType(), Observer->GetObservedOperation(), *Observer);
+			if (UTypedElementQueryObserverCallbackAdapterProcessorBase* Observer =
+				static_cast<UTypedElementQueryObserverCallbackAdapterProcessorBase*>(StoredQuery.Processor.Get()))
+			{
+				Observer->ConfigureQueryCallback(StoredQuery, *this);
+				EntityManager.GetObserverManager().AddObserverInstance(*Observer->GetObservedType(), Observer->GetObservedOperation(), *Observer);
+			}
+			else
+			{
+				return false;
+			}
 		}
 		else
 		{
diff --git a/Engine/Source/Runtime/TypedElementFramework/Private/Elements/Framework/TypedElementQueryBuilder.cpp b/Engine/Source/Runtime/TypedElementFramework/Private/Elements/Framework/TypedElementQueryBuilder.cpp
--- a/Engine/Source/Runtime/TypedElementFramework/Private/Elements/Framework/TypedElementQueryBuilder.cpp
+++ b/Engine/Source/Runtime/TypedElementFramework/Private/Elements/Framework/TypedElementQueryBuilder.cpp
@@ -45,8 +45,13 @@ namespace TypedElementQueryBuilder
 		return *this;
 	}
 
-	FDependency& FDependency::ReadOnly(std::initializer_list<const UClass*> Targets)
+	FDependency& FDependency::ReadOnly(TConstArrayView<const UClass*> Targets)
 	{
+		int32 NewSize = Query->CachedDependencies.Num() + Targets.Num();
+		Query->DependencyTypes.Reserve(NewSize);
+		Query->CachedDependencies.Reserve(NewSize);
+		Query->DependencyFlags.Reserve(NewSize);
+		
 		for (const UClass* Target : Targets)
 		{
 			ReadOnly(Target);
@@ -63,8 +68,13 @@ namespace TypedElementQueryBuilder
 		return *this;
 	}
 
-	FDependency& FDependency::ReadWrite(std::initializer_list<const UClass*> Targets)
+	FDependency& FDependency::ReadWrite(TConstArrayView<const UClass*> Targets)
 	{
+		int32 NewSize = Query->CachedDependencies.Num() + Targets.Num();
+		Query->DependencyTypes.Reserve(NewSize);
+		Query->CachedDependencies.Reserve(NewSize);
+		Query->DependencyFlags.Reserve(NewSize);
+		
 		for (const UClass* Target : Targets)
 		{
 			ReadWrite(Target);
@@ -72,6 +82,18 @@ namespace TypedElementQueryBuilder
 		return *this;
 	}
 
+	FDependency& FDependency::SubQuery(TypedElementQueryHandle Handle)
+	{
+		Query->Subqueries.Add(Handle);
+		return *this;
+	}
+	
+	FDependency& FDependency::SubQuery(TConstArrayView<TypedElementQueryHandle> Handles)
+	{
+		Query->Subqueries.Insert(Handles.GetData(), Handles.Num(), Query->Subqueries.Num());
+		return *this;
+	}
+
 	ITypedElementDataStorageInterface::FQueryDescription&& FDependency::Compile()
 	{
 		return MoveTemp(*Query);
@@ -98,8 +120,12 @@ namespace TypedElementQueryBuilder
 		return *this;
 	}
 
-	FSimpleQuery& FSimpleQuery::All(std::initializer_list<const UScriptStruct*> Targets)
+	FSimpleQuery& FSimpleQuery::All(TConstArrayView<const UScriptStruct*> Targets)
 	{
+		int32 NewSize = Query->ConditionTypes.Num() + Targets.Num();
+		Query->ConditionTypes.Reserve(NewSize);
+		Query->ConditionOperators.Reserve(NewSize);
+		
 		for (const UScriptStruct* Target : Targets)
 		{
 			All(Target);
@@ -117,8 +143,12 @@ namespace TypedElementQueryBuilder
 		return *this;
 	}
 
-	FSimpleQuery& FSimpleQuery::Any(std::initializer_list<const UScriptStruct*> Targets)
+	FSimpleQuery& FSimpleQuery::Any(TConstArrayView<const UScriptStruct*> Targets)
 	{
+		int32 NewSize = Query->ConditionTypes.Num() + Targets.Num();
+		Query->ConditionTypes.Reserve(NewSize);
+		Query->ConditionOperators.Reserve(NewSize);
+
 		for (const UScriptStruct* Target : Targets)
 		{
 			Any(Target);
@@ -136,8 +166,12 @@ namespace TypedElementQueryBuilder
 		return *this;
 	}
 
-	FSimpleQuery& FSimpleQuery::None(std::initializer_list<const UScriptStruct*> Targets)
+	FSimpleQuery& FSimpleQuery::None(TConstArrayView<const UScriptStruct*> Targets)
 	{
+		int32 NewSize = Query->ConditionTypes.Num() + Targets.Num();
+		Query->ConditionTypes.Reserve(NewSize);
+		Query->ConditionOperators.Reserve(NewSize);
+
 		for (const UScriptStruct* Target : Targets)
 		{
 			None(Target);
@@ -152,6 +186,16 @@ namespace TypedElementQueryBuilder
 
 	ITypedElementDataStorageInterface::FQueryDescription&& FSimpleQuery::Compile()
 	{
+		Query->Callback.BeforeGroups.Shrink();
+		Query->Callback.AfterGroups.Shrink();
+		Query->SelectionTypes.Shrink();
+		Query->SelectionAccessTypes.Shrink();
+		Query->ConditionTypes.Shrink();
+		Query->ConditionOperators.Shrink();
+		Query->DependencyTypes.Shrink();
+		Query->DependencyFlags.Shrink();
+		Query->CachedDependencies.Shrink();
+		Query->Subqueries.Shrink();
 		return MoveTemp(*Query);
 	}
 
diff --git a/Engine/Source/Runtime/TypedElementFramework/Public/Elements/Framework/TypedElementQueryBuilder.h b/Engine/Source/Runtime/TypedElementFramework/Public/Elements/Framework/TypedElementQueryBuilder.h
--- a/Engine/Source/Runtime/TypedElementFramework/Public/Elements/Framework/TypedElementQueryBuilder.h
+++ b/Engine/Source/Runtime/TypedElementFramework/Public/Elements/Framework/TypedElementQueryBuilder.h
@@ -118,11 +118,14 @@ namespace TypedElementQueryBuilder
 		template<typename... TargetTypes>
 		FDependency& ReadOnly();
 		FDependency& ReadOnly(const UClass* Target);
-		FDependency& ReadOnly(std::initializer_list<const UClass*> Targets);
+		FDependency& ReadOnly(TConstArrayView<const UClass*> Targets);
 		template<typename... TargetTypes>
 		FDependency& ReadWrite();
 		FDependency& ReadWrite(const UClass* Target);
-		FDependency& ReadWrite(std::initializer_list<const UClass*> Targets);
+		FDependency& ReadWrite(TConstArrayView<const UClass*> Targets);
+
+		FDependency& SubQuery(TypedElementQueryHandle Handle);
+		FDependency& SubQuery(TConstArrayView<TypedElementQueryHandle> Handles);
 
 		ITypedElementDataStorageInterface::FQueryDescription&& Compile();
 
@@ -144,15 +147,15 @@ namespace TypedElementQueryBuilder
 		template<typename... TargetTypes>
 		FSimpleQuery& All();
 		FSimpleQuery& All(const UScriptStruct* Target);
-		FSimpleQuery& All(std::initializer_list<const UScriptStruct*> Targets);
+		FSimpleQuery& All(TConstArrayView<const UScriptStruct*> Targets);
 		template<typename... TargetTypes>
 		FSimpleQuery& Any();
 		FSimpleQuery& Any(const UScriptStruct* Target);
-		FSimpleQuery& Any(std::initializer_list<const UScriptStruct*> Targets);
+		FSimpleQuery& Any(TConstArrayView<const UScriptStruct*> Targets);
 		template<typename... TargetTypes>
 		FSimpleQuery& None();
 		FSimpleQuery& None(const UScriptStruct* Target);
-		FSimpleQuery& None(std::initializer_list<const UScriptStruct*> Targets);
+		FSimpleQuery& None(TConstArrayView<const UScriptStruct*> Targets);
 
 	private:
 		explicit FSimpleQuery(ITypedElementDataStorageInterface::FQueryDescription* Query);
@@ -251,6 +254,7 @@ namespace TypedElementQueryBuilder
 		inline void RemoveColumns(TConstArrayView<TypedElementRowHandle> Rows, TConstArrayView<const UScriptStruct*> ColumnTypes) override;
 
 		inline ITypedElementDataStorageInterface::FQueryResult RunQuery(TypedElementQueryHandle Query) override;
+		inline ITypedElementDataStorageInterface::FQueryResult RunSubquery(int32 SubqueryIndex) override;
 
 		ITypedElementDataStorageInterface::IQueryContext& ParentContext;
 		const ITypedElementDataStorageInterface::FQueryDescription& Description;
diff --git a/Engine/Source/Runtime/TypedElementFramework/Public/Elements/Framework/TypedElementQueryBuilder.inl b/Engine/Source/Runtime/TypedElementFramework/Public/Elements/Framework/TypedElementQueryBuilder.inl
--- a/Engine/Source/Runtime/TypedElementFramework/Public/Elements/Framework/TypedElementQueryBuilder.inl
+++ b/Engine/Source/Runtime/TypedElementFramework/Public/Elements/Framework/TypedElementQueryBuilder.inl
@@ -192,6 +192,11 @@ namespace TypedElementQueryBuilder
 	{
 		return ParentContext.RunQuery(Query);
 	}
+	
+	ITypedElementDataStorageInterface::FQueryResult FQueryContextForwarder::RunSubquery(int32 SubqueryIndex)
+	{
+		return ParentContext.RunSubquery(SubqueryIndex);
+	}
 
 	//
 	// FCachedQueryContext
diff --git a/Engine/Source/Runtime/TypedElementFramework/Public/Elements/Interfaces/TypedElementDataStorageInterface.h b/Engine/Source/Runtime/TypedElementFramework/Public/Elements/Interfaces/TypedElementDataStorageInterface.h
--- a/Engine/Source/Runtime/TypedElementFramework/Public/Elements/Interfaces/TypedElementDataStorageInterface.h
+++ b/Engine/Source/Runtime/TypedElementFramework/Public/Elements/Interfaces/TypedElementDataStorageInterface.h
@@ -392,6 +392,10 @@ public:
 		 * access data from rows such as count queries.
 		 */
 		virtual FQueryResult RunQuery(TypedElementQueryHandle Query) = 0;
+		/** Runs a subquery registered with the current query. The subquery index is in the order of registration with the query. */
+		virtual FQueryResult RunSubquery(int32 SubqueryIndex) = 0;
+		
+
 
 		// Utility functions
 
@@ -406,13 +410,21 @@ public:
 	};
 
 	/**
-	 * Interface to be provided to query callbacks that are directly called through RunQuery.
+	 * Interface to be provided to query callbacks that are directly called through RunQuery from outside a query callback.
 	 */
 	struct IDirectQueryContext : public ICommonQueryContext
 	{
 		virtual ~IDirectQueryContext() = default;
 	};
 
+	/**
+	 * Interface to be provided to query callbacks that are directly called through from a query callback.
+	 */
+	struct ISubqueryCallbackContext : public ICommonQueryContext
+	{
+		virtual ~ISubqueryCallbackContext() = default;
+	};
+
 	struct FQueryDescription;
 	using QueryCallback = TFunction<void(const FQueryDescription&, IQueryContext&)>;
 	using QueryCallbackRef = TFunctionRef<void(const FQueryDescription&, IQueryContext&)>;
@@ -490,6 +502,7 @@ public:
 		TArray<EQueryDependencyFlags, TInlineAllocator<NumInlineDependencies>> DependencyFlags;
 		/** Cached instances of the dependencies. This will always match the count of the other Dependency*Types, but may contain null pointers. */
 		TArray<TWeakObjectPtr<UObject>, TInlineAllocator<NumInlineDependencies>> CachedDependencies;
+		TArray<TypedElementQueryHandle> Subqueries;
 		
 		EActionType Action;
 		/** If true, this query only has simple operations and is guaranteed to be executed fully and at optimal performance. */
