Added support to TEDS to reserve a row that's assigned to a table at a later point in time.

This change allows reserving a row to be assigned later. For time sliced work this is useful as it allows the row handle to be used before the table has been fully initialized. Note though that the creation/reservation of a row is not thread safe at this point in time, so this can only be safely used for time sliced solutions, not for threading solutions yet.
To reduce the number of overloaded functions, the ability to use a table by name has been removed. Instead call `FindTable` to get the table handle.

diff --git a/Engine/Plugins/Experimental/TypedElementsDataStorage/Source/TypedElementsDataStorage/Private/TypedElementDatabase.cpp b/Engine/Plugins/Experimental/TypedElementsDataStorage/Source/TypedElementsDataStorage/Private/TypedElementDatabase.cpp
--- a/Engine/Plugins/Experimental/TypedElementsDataStorage/Source/TypedElementsDataStorage/Private/TypedElementDatabase.cpp
+++ b/Engine/Plugins/Experimental/TypedElementsDataStorage/Source/TypedElementsDataStorage/Private/TypedElementDatabase.cpp
@@ -122,14 +122,9 @@ TSharedPtr<const FMassEntityManager> UTypedElementDatabase::GetActiveEditorEntit
 	return ActiveEditorEntityManager;
 }
 
-TypedElementTableHandle UTypedElementDatabase::RegisterTable(TConstArrayView<const UScriptStruct*> ColumnList)
-{
-	return RegisterTable(ColumnList, {});
-}
-
 TypedElementTableHandle UTypedElementDatabase::RegisterTable(TConstArrayView<const UScriptStruct*> ColumnList, const FName Name)
 {
-	if (ActiveEditorEntityManager && (!Name.IsValid() || !TableNameLookup.Contains(Name)))
+	if (ActiveEditorEntityManager && !TableNameLookup.Contains(Name))
 	{
 		TypedElementTableHandle Result = Tables.Num();
 		Tables.Add(ActiveEditorEntityManager->CreateArchetype(ColumnList, Name));
@@ -142,16 +137,10 @@ TypedElementTableHandle UTypedElementDatabase::RegisterTable(TConstArrayView<con
 	return TypedElementInvalidTableHandle;
 }
 
-TypedElementTableHandle UTypedElementDatabase::RegisterTable(TypedElementTableHandle SourceTable,
-	TConstArrayView<const UScriptStruct*> ColumnList)
-{
-	return RegisterTable(SourceTable, ColumnList, {});
-}
-
 TypedElementTableHandle UTypedElementDatabase::RegisterTable(TypedElementTableHandle SourceTable, 
 	TConstArrayView<const UScriptStruct*> ColumnList, const FName Name)
 {
-	if (ActiveEditorEntityManager && (!Name.IsValid() || !TableNameLookup.Contains(Name)) && SourceTable < Tables.Num())
+	if (ActiveEditorEntityManager && SourceTable < Tables.Num() && !TableNameLookup.Contains(Name))
 	{
 		TypedElementTableHandle Result = Tables.Num();
 		Tables.Add(ActiveEditorEntityManager->CreateArchetype(Tables[SourceTable], ColumnList, Name));
@@ -170,6 +159,11 @@ TypedElementTableHandle UTypedElementDatabase::FindTable(const FName Name)
 	return TableHandle ? *TableHandle : TypedElementInvalidTableHandle;
 }
 
+TypedElementRowHandle UTypedElementDatabase::ReserveRow()
+{
+	return ActiveEditorEntityManager ? ActiveEditorEntityManager->ReserveEntity().AsNumber(): TypedElementInvalidRowHandle;
+}
+
 TypedElementRowHandle UTypedElementDatabase::AddRow(TypedElementTableHandle Table)
 {
 	checkf(Table < Tables.Num(), TEXT("Attempting to add a row to a non-existing table."));
@@ -178,10 +172,19 @@ TypedElementRowHandle UTypedElementDatabase::AddRow(TypedElementTableHandle Tabl
 		TypedElementInvalidRowHandle;
 }
 
-TypedElementRowHandle UTypedElementDatabase::AddRow(FName TableName)
+bool UTypedElementDatabase::AddRow(TypedElementRowHandle ReservedRow, TypedElementTableHandle Table)
 {
-	TypedElementTableHandle* Table = TableNameLookup.Find(TableName);
-	return Table ? AddRow(*Table) : TypedElementInvalidRowHandle;
+	checkf(!HasRowBeenAssigned(ReservedRow), TEXT("Attempting to assign a table to row that already has a table assigned."));
+	checkf(Table < Tables.Num(), TEXT("Attempting to add a row to a non-existing table."));
+	if (ActiveEditorEntityManager)
+	{
+		ActiveEditorEntityManager->BuildEntity(FMassEntityHandle::FromNumber(ReservedRow), Tables[Table]);
+		return true;
+	}
+	else
+	{
+		return false;
+	}
 }
 
 bool UTypedElementDatabase::BatchAddRow(TypedElementTableHandle Table, int32 Count, TypedElementDataStorageCreationCallbackRef OnCreated)
@@ -205,20 +208,61 @@ bool UTypedElementDatabase::BatchAddRow(TypedElementTableHandle Table, int32 Cou
 	return false;
 }
 
-bool UTypedElementDatabase::BatchAddRow(FName TableName, int32 Count, TypedElementDataStorageCreationCallbackRef OnCreated)
+bool UTypedElementDatabase::BatchAddRow(TypedElementTableHandle Table, TConstArrayView<TypedElementRowHandle> ReservedHandles,
+	TypedElementDataStorageCreationCallbackRef OnCreated)
 {
-	TypedElementTableHandle* Table = TableNameLookup.Find(TableName);
-	return Table ? BatchAddRow(*Table, Count, OnCreated) : false;
+	OnCreated.CheckCallable();
+	checkf(Table < Tables.Num(), TEXT("Attempting to add multiple rows to a non-existing table."));
+	if (ActiveEditorEntityManager)
+	{
+		TArray<FMassEntityHandle> Entities;
+		Entities.Reserve(ReservedHandles.Num());
+		FMassEntityHandle* CurrentEntityHandle = Entities.GetData();
+		for (TypedElementRowHandle RowHandle : ReservedHandles)
+		{
+			checkf(!HasRowBeenAssigned(RowHandle), TEXT("Attempting to assign a table to row that already has a table assigned."));
+			Entities.Add(FMassEntityHandle::FromNumber(RowHandle));
+		}
+		
+		TSharedRef<FMassEntityManager::FEntityCreationContext> Context =
+			ActiveEditorEntityManager->BatchCreateReservedEntities(Tables[Table], Entities);
+
+		for (FMassEntityHandle Entity : Entities)
+		{
+			OnCreated(Entity.AsNumber());
+		}
+
+		return true;
+	}
+	return false;
 }
 
+
 void UTypedElementDatabase::RemoveRow(TypedElementRowHandle Row)
 {
 	if (ActiveEditorEntityManager)
 	{
-		ActiveEditorEntityManager->DestroyEntity(FMassEntityHandle::FromNumber(Row));
+		if (ActiveEditorEntityManager->IsEntityActive(FMassEntityHandle::FromNumber(Row)))
+		{
+			ActiveEditorEntityManager->DestroyEntity(FMassEntityHandle::FromNumber(Row));
+		}
+		else
+		{
+			ActiveEditorEntityManager->ReleaseReservedEntity(FMassEntityHandle::FromNumber(Row));
+		}
 	}
 }
 
+bool UTypedElementDatabase::IsRowAvailable(TypedElementRowHandle Row) const
+{
+	return ActiveEditorEntityManager ? ActiveEditorEntityManager->IsEntityValid(FMassEntityHandle::FromNumber(Row)) : false;
+}
+
+bool UTypedElementDatabase::HasRowBeenAssigned(TypedElementRowHandle Row) const
+{
+	return ActiveEditorEntityManager ? ActiveEditorEntityManager->IsEntityActive(FMassEntityHandle::FromNumber(Row)) : false;
+}
+
 bool UTypedElementDatabase::AddColumn(TypedElementRowHandle Row, const UScriptStruct* ColumnType)
 {
 	FMassEntityHandle Entity = FMassEntityHandle::FromNumber(Row);
diff --git a/Engine/Plugins/Experimental/TypedElementsDataStorage/Source/TypedElementsDataStorage/Private/TypedElementDatabase.h b/Engine/Plugins/Experimental/TypedElementsDataStorage/Source/TypedElementsDataStorage/Private/TypedElementDatabase.h
--- a/Engine/Plugins/Experimental/TypedElementsDataStorage/Source/TypedElementsDataStorage/Private/TypedElementDatabase.h
+++ b/Engine/Plugins/Experimental/TypedElementsDataStorage/Source/TypedElementsDataStorage/Private/TypedElementDatabase.h
@@ -35,18 +35,20 @@ public:
 	TSharedPtr<FMassEntityManager> GetActiveMutableEditorEntityManager();
 	TSharedPtr<const FMassEntityManager> GetActiveEditorEntityManager() const;
 
-	TypedElementTableHandle RegisterTable(TConstArrayView<const UScriptStruct*> ColumnList) override;
 	TypedElementTableHandle RegisterTable(TConstArrayView<const UScriptStruct*> ColumnList, const FName Name) override;
-	TypedElementTableHandle RegisterTable(TypedElementTableHandle SourceTable, TConstArrayView<const UScriptStruct*> ColumnList) override;
 	TypedElementTableHandle RegisterTable(TypedElementTableHandle SourceTable, TConstArrayView<const UScriptStruct*> ColumnList, 
 		const FName Name) override;
 	TypedElementTableHandle FindTable(const FName Name) override;
 
+	TypedElementRowHandle ReserveRow() override;
 	TypedElementRowHandle AddRow(TypedElementTableHandle Table) override;
-	TypedElementRowHandle AddRow(FName TableName) override;
+	bool AddRow(TypedElementRowHandle ReservedRow, TypedElementTableHandle Table) override;
 	bool BatchAddRow(TypedElementTableHandle Table, int32 Count, TypedElementDataStorageCreationCallbackRef OnCreated) override;
-	bool BatchAddRow(FName TableName, int32 Count, TypedElementDataStorageCreationCallbackRef OnCreated) override;
+	bool BatchAddRow(TypedElementTableHandle Table, TConstArrayView<TypedElementRowHandle> ReservedHandles,
+		TypedElementDataStorageCreationCallbackRef OnCreated) override;
 	void RemoveRow(TypedElementRowHandle Row) override;
+	bool IsRowAvailable(TypedElementRowHandle Row) const override;
+	bool HasRowBeenAssigned(TypedElementRowHandle Row) const override;
 
 	bool AddColumn(TypedElementRowHandle Row, const UScriptStruct* ColumnType) override;
 	bool AddColumn(TypedElementRowHandle Row, FTopLevelAssetPath ColumnName) override;
diff --git a/Engine/Plugins/Experimental/TypedElementsDataStorage/Source/TypedElementsDataStorageUI/Private/TypedElementSceneOutliner.cpp b/Engine/Plugins/Experimental/TypedElementsDataStorage/Source/TypedElementsDataStorageUI/Private/TypedElementSceneOutliner.cpp
--- a/Engine/Plugins/Experimental/TypedElementsDataStorage/Source/TypedElementsDataStorageUI/Private/TypedElementSceneOutliner.cpp
+++ b/Engine/Plugins/Experimental/TypedElementsDataStorage/Source/TypedElementsDataStorageUI/Private/TypedElementSceneOutliner.cpp
@@ -151,7 +151,7 @@ public:
 		TSharedPtr<SWidget> Widget;
 		if (HeaderWidgetConstructor)
 		{
-			TypedElementRowHandle UiRowHandle = Storage.AddRow(FName(TEXT("Editor_WidgetTable")));
+			TypedElementRowHandle UiRowHandle = Storage.AddRow(Storage.FindTable(FName(TEXT("Editor_WidgetTable"))));
 			Widget = StorageUi.ConstructWidget(UiRowHandle, *HeaderWidgetConstructor, 
 				{
 					{ FName(TEXT("Name")), NameId.ToString() }
@@ -187,7 +187,7 @@ public:
 				TypedElementRowHandle RowHandle = StorageCompatibility.FindRowWithCompatibleObject(Actor);
 				if (RowHandle != TypedElementInvalidRowHandle && Storage.HasColumns(RowHandle, ColumnTypes))
 				{
-					TypedElementRowHandle UiRowHandle = Storage.AddRow(FTypedElementSceneOutlinerQueryBinder::CellWidgetTableName);
+					TypedElementRowHandle UiRowHandle = Storage.AddRow(Storage.FindTable(FTypedElementSceneOutlinerQueryBinder::CellWidgetTableName));
 					Storage.AddColumns(UiRowHandle, CellWidgetConstructor->GetAdditionalColumnsList());
 					if (ColumnTypes.Num() == 1)
 					{
diff --git a/Engine/Source/Runtime/TypedElementFramework/Public/Elements/Interfaces/TypedElementDataStorageInterface.h b/Engine/Source/Runtime/TypedElementFramework/Public/Elements/Interfaces/TypedElementDataStorageInterface.h
--- a/Engine/Source/Runtime/TypedElementFramework/Public/Elements/Interfaces/TypedElementDataStorageInterface.h
+++ b/Engine/Source/Runtime/TypedElementFramework/Public/Elements/Interfaces/TypedElementDataStorageInterface.h
@@ -94,14 +94,11 @@ public:
 	 */
 
 	/** Creates a new table for with the provided columns. Optionally a name can be given which is useful for retrieval later. */
-	virtual TypedElementTableHandle RegisterTable(TConstArrayView<const UScriptStruct*> ColumnList) = 0;
 	virtual TypedElementTableHandle RegisterTable(TConstArrayView<const UScriptStruct*> ColumnList, const FName Name) = 0;
 	/** 
 	 * Copies the column information from the provided table and creates a new table for with the provided columns. Optionally a 
 	 * name can be given which is useful for retrieval later.
 	 */
-	virtual TypedElementTableHandle RegisterTable(TypedElementTableHandle SourceTable, 
-		TConstArrayView<const UScriptStruct*> ColumnList) = 0;
 	virtual TypedElementTableHandle RegisterTable(TypedElementTableHandle SourceTable, 
 		TConstArrayView<const UScriptStruct*> ColumnList, const FName Name) = 0;
 
@@ -112,20 +109,37 @@ public:
 	 * @section Row management
 	 */
 
+	/** 
+	 * Reserves a row to be assigned to a table at a later point. If the row is no longer needed before it's been assigned
+	 * to a table, it should still be released with RemoveRow.
+	 */
+	virtual TypedElementRowHandle ReserveRow() = 0;
 	/** Adds a new row to the provided table. */
 	virtual TypedElementRowHandle AddRow(TypedElementTableHandle Table) = 0;
-	virtual TypedElementRowHandle AddRow(FName TableName) = 0;
+	/** Adds a new row to the provided table using a previously reserved row.. */
+	virtual bool AddRow(TypedElementRowHandle ReservedRow, TypedElementTableHandle Table) = 0;
 	
 	/**
 	 * Add multiple rows at once. For each new row the OnCreated callback is called. Callers are expected to use the callback to
 	 * initialize the row if needed.
 	 */
 	virtual bool BatchAddRow(TypedElementTableHandle Table, int32 Count, TypedElementDataStorageCreationCallbackRef OnCreated) = 0;
-	virtual bool BatchAddRow(FName TableName, int32 Count, TypedElementDataStorageCreationCallbackRef OnCreated) = 0;
-	
-	/** Removes a previously added row. If the row handle is invalid or already removed, nothing happens */
+	/**
+	 * Add multiple rows at once. For each new row the OnCreated callback is called. Callers are expected to use the callback to
+	 * initialize the row if needed. This version uses a set of previously reserved rows. Any row that can't be used will be 
+	 * released.
+	 */
+	virtual bool BatchAddRow(TypedElementTableHandle Table, TConstArrayView<TypedElementRowHandle> ReservedHandles,
+		TypedElementDataStorageCreationCallbackRef OnCreated) = 0;
+
+	/** Removes a previously reserved or added row. If the row handle is invalid or already removed, nothing happens */
 	virtual void RemoveRow(TypedElementRowHandle Row) = 0;
 
+	/** Checks whether or not a row is in use. This is true even if the row has only been reserved. */
+	virtual bool IsRowAvailable(TypedElementRowHandle Row) const = 0;
+	/** Checks whether or not a row has been reserved but not yet assigned to a table. */
+	virtual bool HasRowBeenAssigned(TypedElementRowHandle Row) const = 0;
+
 	
 	/**
 	 * @section Column management
@@ -641,7 +655,7 @@ ColumnType* ITypedElementDataStorageInterface::AddOrGetColumn(TypedElementRowHan
 	{
 		if (Result)
 		{
-			new(Result) ColumnType{ Forward<Arguments>... };
+			new(Result) ColumnType{ std::forward<Args>(Arguments)... };
 		}
 	}
 	return Result;
