[TEDS] Local Transform Column as heads up display for abnormal transform attributes

Supports non-uniform scale, unnormalized rotation, and negative scaling on XYZ axis

diff --git a/Engine/Plugins/Experimental/TypedElementsDataStorage/Source/TypedElementsDataStorageUI/Private/TypedElementTransformHeadsUpWidget.cpp b/Engine/Plugins/Experimental/TypedElementsDataStorage/Source/TypedElementsDataStorageUI/Private/TypedElementTransformHeadsUpWidget.cpp
new file mode 100644
--- /dev/null
+++ b/Engine/Plugins/Experimental/TypedElementsDataStorage/Source/TypedElementsDataStorageUI/Private/TypedElementTransformHeadsUpWidget.cpp
@@ -0,0 +1,240 @@
+// Copyright Epic Games, Inc. All Rights Reserved.
+
+#include "TypedElementTransformHeadsUpWidget.h"
+
+#include "Elements/Columns/TypedElementMiscColumns.h"
+#include "Elements/Columns/TypedElementSlateWidgetColumns.h"
+#include "Elements/Columns/TypedElementTransformColumns.h"
+#include "Elements/Framework/TypedElementQueryBuilder.h"
+#include "TypedElementSubsystems.h"
+#include "Widgets/SCompoundWidget.h"
+#include "Widgets/SBoxPanel.h"
+#include "Widgets/Images/SImage.h"
+#include "Styling/AppStyle.h"
+#include "Widgets/Input/SNumericEntryBox.h"
+
+#define LOCTEXT_NAMESPACE "TypedElementUI_TransformHeadsUpWidget"
+
+namespace InternalTransformHelpers
+{
+	enum class EAbnormalTransformTypes : int32
+	{
+		None = 0x0000,
+		NonUniformScale = 0x0001,
+		NegativeXScale = 0x0002,
+		NegativeYScale = 0x0004,
+		NegativeZScale = 0x0008,
+		UnnormalizedRotation = 0x0010,
+	};
+	ENUM_CLASS_FLAGS(EAbnormalTransformTypes);
+
+	EAbnormalTransformTypes GetAbnormalTransformTypes(const FTransform& InTransform)
+	{
+		EAbnormalTransformTypes AbnormalTransformFlags = EAbnormalTransformTypes::None;
+		const FVector& TransformScale = InTransform.GetScale3D();
+
+		if (!TransformScale.GetAbs().AllComponentsEqual())
+		{
+			AbnormalTransformFlags |= EAbnormalTransformTypes::NonUniformScale;
+		}
+		if (TransformScale.X < 0.0)
+		{
+			AbnormalTransformFlags |= EAbnormalTransformTypes::NegativeXScale;
+		}
+		if (TransformScale.Y < 0.0)
+		{
+			AbnormalTransformFlags |= EAbnormalTransformTypes::NegativeYScale;
+		}
+		if (TransformScale.Z < 0.0)
+		{
+			AbnormalTransformFlags |= EAbnormalTransformTypes::NegativeZScale;
+		}
+		if (!InTransform.IsRotationNormalized())
+		{
+			AbnormalTransformFlags |= EAbnormalTransformTypes::UnnormalizedRotation;
+		}
+
+		return AbnormalTransformFlags;
+	}
+} // namespace InternalTransformHelpers
+
+class STransformQuickDisplay : public SHorizontalBox
+{
+	SLATE_DECLARE_WIDGET(STransformQuickDisplay, SHorizontalBox)
+public:
+	SLATE_BEGIN_ARGS(STransformQuickDisplay) {}
+	SLATE_END_ARGS()
+
+	void Construct(const FArguments& InArgs)
+	{
+		const static FMargin IconPadding(1, 1, 0, 0);
+		using namespace InternalTransformHelpers;
+
+		AddSlot().AutoWidth().Padding(IconPadding)
+		[
+			SNew(SImage)
+				.Image(FAppStyle::GetBrush("EditorViewport.ScaleGridSnap"))
+				.ToolTipText(LOCTEXT("NonUniformScaleTooltip", "Has Non-Uniform Scale"))
+				.Visibility_Lambda([this]() { return EnumHasAllFlags(AbnormalTransformFlags, EAbnormalTransformTypes::NonUniformScale) ? EVisibility::HitTestInvisible : EVisibility::Collapsed; })
+		];
+		AddSlot().AutoWidth().Padding(IconPadding)
+		[
+			SNew(SImage)
+				.Image(FAppStyle::GetBrush("EditorViewport.ScaleMode"))
+				.ColorAndOpacity(SNumericEntryBox<double>::RedLabelBackgroundColor)
+				.ToolTipText(LOCTEXT("NegativeXScaleTooltip", "Has Negative X Scale"))
+				.Visibility_Lambda([this]() { return EnumHasAllFlags(AbnormalTransformFlags, EAbnormalTransformTypes::NegativeXScale) ? EVisibility::HitTestInvisible : EVisibility::Collapsed; })
+		];
+		AddSlot().AutoWidth().Padding(IconPadding)
+		[
+			SNew(SImage)
+				.Image(FAppStyle::GetBrush("EditorViewport.ScaleMode"))
+				.ColorAndOpacity(SNumericEntryBox<double>::GreenLabelBackgroundColor)
+				.ToolTipText(LOCTEXT("NegativeYScaleTooltip", "Has Negative Y Scale"))
+				.Visibility_Lambda([this]() { return EnumHasAllFlags(AbnormalTransformFlags, EAbnormalTransformTypes::NegativeYScale) ? EVisibility::HitTestInvisible : EVisibility::Collapsed; })
+		];
+		AddSlot().AutoWidth().Padding(IconPadding)
+		[
+			SNew(SImage)
+				.Image(FAppStyle::GetBrush("EditorViewport.ScaleMode"))
+				.ColorAndOpacity(SNumericEntryBox<double>::BlueLabelBackgroundColor)
+				.ToolTipText(LOCTEXT("NegativeZScaleTooltip", "Has Negative Z Scale"))
+				.Visibility_Lambda([this]() { return EnumHasAllFlags(AbnormalTransformFlags, EAbnormalTransformTypes::NegativeZScale) ? EVisibility::HitTestInvisible : EVisibility::Collapsed; })
+		];
+		AddSlot().AutoWidth().Padding(IconPadding)
+		[
+			SNew(SImage)
+				.Image(FAppStyle::GetBrush("SurfaceDetails.AntiClockwiseRotation"))
+				.ToolTipText(LOCTEXT("UnnormalizedRotationTooltip", "Has Un-normalized Rotation"))
+				.Visibility_Lambda([this]() { return EnumHasAllFlags(AbnormalTransformFlags, EAbnormalTransformTypes::UnnormalizedRotation) ? EVisibility::HitTestInvisible : EVisibility::Collapsed; })
+		];
+		AddSlot().AutoWidth().Padding(IconPadding)
+		[
+			SNew(SImage)
+				.Image(FAppStyle::GetBrush("Symbols.Check"))
+				.ToolTipText(LOCTEXT("NothingToReportTooltip", "No Abnormal Transform Data"))
+				.Visibility_Lambda([this]() { return (AbnormalTransformFlags == EAbnormalTransformTypes::None) ? EVisibility::HitTestInvisible : EVisibility::Collapsed; })
+		];
+	}
+
+	void UpdateFromTransform(InternalTransformHelpers::EAbnormalTransformTypes InAbnormalTransformFlags)
+	{
+		if (AbnormalTransformFlags != InAbnormalTransformFlags)
+		{
+			AbnormalTransformFlags = InAbnormalTransformFlags;
+			Invalidate(EInvalidateWidgetReason::Visibility);
+		}
+	}
+
+private:
+	InternalTransformHelpers::EAbnormalTransformTypes AbnormalTransformFlags;
+};
+
+SLATE_IMPLEMENT_WIDGET(STransformQuickDisplay)
+void STransformQuickDisplay::PrivateRegisterAttributes(FSlateAttributeInitializer&) {}
+
+static void UpdateTransformHeadsUpDisplay(ITypedElementDataStorageInterface& DataStorage, FTypedElementSlateWidgetReferenceColumn& Widget, InternalTransformHelpers::EAbnormalTransformTypes AbnormalTransformFlags)
+{
+	TSharedPtr<SWidget> WidgetPointer = Widget.Widget.Pin();
+	checkf(WidgetPointer, TEXT("Referenced widget is not valid. A constructed widget may not have been cleaned up. This can "
+		"also happen if this processor is running in the same phase as the processors responsible for cleaning up old "
+		"references."));
+	checkf(WidgetPointer->GetType() == STransformQuickDisplay::StaticWidgetClass().GetWidgetType(),
+		TEXT("Stored widget with FTypedElementTransformHeadsUpWidgetTag doesn't match type %s, but was a %s."),
+		*(STransformQuickDisplay::StaticWidgetClass().GetWidgetType().ToString()),
+		*(WidgetPointer->GetTypeAsString()));
+
+	STransformQuickDisplay* BoxWidget = static_cast<STransformQuickDisplay*>(WidgetPointer.Get());
+	BoxWidget->UpdateFromTransform(AbnormalTransformFlags);
+}
+
+//
+// UTypedElementTransformHeadsUpWidgetFactory
+//
+
+void UTypedElementTransformHeadsUpWidgetFactory::RegisterQueries(ITypedElementDataStorageInterface& DataStorage) const
+{
+	using namespace TypedElementQueryBuilder;
+	using DSI = ITypedElementDataStorageInterface;
+		
+
+	DataStorage.RegisterQuery(
+		Select(TEXT("Sync Transform column to heads up display"),
+		FProcessor(DSI::EQueryTickPhase::FrameEnd, DataStorage.GetQueryTickGroupName(DSI::EQueryTickGroups::SyncWidgets))
+			.ForceToGameThread(true),
+			[](FCachedQueryContext<UTypedElementDataStorageSubsystem>& Context,
+				FTypedElementSlateWidgetReferenceColumn& Widget,
+				const FTypedElementRowReferenceColumn& ReferenceColumn)
+			{
+				UTypedElementDataStorageSubsystem& Subsystem = Context.GetCachedMutableDependency<UTypedElementDataStorageSubsystem>();
+				DSI* DataStorage = Subsystem.Get();
+				checkf(DataStorage, TEXT("FTypedElementsDataStorageUiModule tried to process widgets before the "
+					"Typed Elements Data Storage interface is available."));
+
+				if (DataStorage->HasColumns<FTypedElementSyncFromWorldTag>(ReferenceColumn.Row)
+					|| DataStorage->HasColumns<FTypedElementSyncBackToWorldTag>(ReferenceColumn.Row))
+				{
+					if (const FTypedElementLocalTransformColumn* TransformColumn = DataStorage->GetColumn<FTypedElementLocalTransformColumn>(ReferenceColumn.Row))
+					{
+						UpdateTransformHeadsUpDisplay(*DataStorage, Widget, InternalTransformHelpers::GetAbnormalTransformTypes(TransformColumn->Transform));
+					}
+				}
+			}
+		)
+	.Where()
+		.All<FTypedElementTransformHeadsUpWidgetTag>()
+	.Compile());
+
+}
+
+void UTypedElementTransformHeadsUpWidgetFactory::RegisterWidgetConstructors(ITypedElementDataStorageInterface& DataStorage,
+	ITypedElementDataStorageUiInterface& DataStorageUi) const
+{
+	DataStorageUi.RegisterWidgetFactory(FName(TEXT("SceneOutliner.Cell")), FTypedElementTransformHeadsUpWidgetConstructor::StaticStruct(),
+		{ FTypedElementLocalTransformColumn::StaticStruct() });
+}
+
+
+
+//
+// FTypedElementTransformHeadsUpWidgetConstructor
+//
+
+FTypedElementTransformHeadsUpWidgetConstructor::FTypedElementTransformHeadsUpWidgetConstructor()
+	: Super(FTypedElementTransformHeadsUpWidgetConstructor::StaticStruct())
+{
+}
+
+TConstArrayView<const UScriptStruct*> FTypedElementTransformHeadsUpWidgetConstructor::GetAdditionalColumnsList() const
+{
+	static TTypedElementColumnTypeList<FTypedElementRowReferenceColumn,
+		FTypedElementTransformHeadsUpWidgetTag> Columns;
+	return Columns;
+}
+
+bool FTypedElementTransformHeadsUpWidgetConstructor::CanBeReused() const
+{
+	return true;
+}
+
+TSharedPtr<SWidget> FTypedElementTransformHeadsUpWidgetConstructor::CreateWidget()
+{
+	return SNew(STransformQuickDisplay);
+}
+
+bool FTypedElementTransformHeadsUpWidgetConstructor::FinalizeWidget(
+	ITypedElementDataStorageInterface* DataStorage,
+	ITypedElementDataStorageUiInterface* DataStorageUi,
+	TypedElementRowHandle Row,
+	const TSharedPtr<SWidget>& Widget)
+{
+	FTypedElementRowReferenceColumn& RefColumn = *DataStorage->GetColumn<FTypedElementRowReferenceColumn>(Row);
+	if (const FTypedElementLocalTransformColumn* TransformColumn = DataStorage->GetColumn<FTypedElementLocalTransformColumn>(RefColumn.Row))
+	{
+		UpdateTransformHeadsUpDisplay(*DataStorage, *DataStorage->GetColumn<FTypedElementSlateWidgetReferenceColumn>(Row), InternalTransformHelpers::GetAbnormalTransformTypes(TransformColumn->Transform));
+	}
+
+	return true;
+}
+
+#undef LOCTEXT_NAMESPACE
diff --git a/Engine/Plugins/Experimental/TypedElementsDataStorage/Source/TypedElementsDataStorageUI/Public/TypedElementTransformHeadsUpWidget.h b/Engine/Plugins/Experimental/TypedElementsDataStorage/Source/TypedElementsDataStorageUI/Public/TypedElementTransformHeadsUpWidget.h
new file mode 100644
--- /dev/null
+++ b/Engine/Plugins/Experimental/TypedElementsDataStorage/Source/TypedElementsDataStorageUI/Public/TypedElementTransformHeadsUpWidget.h
@@ -0,0 +1,57 @@
+// Copyright Epic Games, Inc. All Rights Reserved.
+
+#pragma once
+
+#include "Elements/Interfaces/TypedElementDataStorageFactory.h"
+#include "Elements/Interfaces/TypedElementDataStorageInterface.h"
+#include "Elements/Interfaces/TypedElementDataStorageUiInterface.h"
+#include "Internationalization/Text.h"
+#include "UObject/ObjectMacros.h"
+
+#include "TypedElementTransformHeadsUpWidget.generated.h"
+
+/**
+ * The heads up transform display provides at a glance view in a scene outliner row of abnormal transform characteristics, including:
+ *		1. Non-uniform scale
+ *		2. Negative scaling on X, Y, or Z axis
+ *		3. Unnormalized rotation
+ */
+UCLASS()
+class TYPEDELEMENTSDATASTORAGEUI_API UTypedElementTransformHeadsUpWidgetFactory : public UTypedElementDataStorageFactory
+{
+	GENERATED_BODY()
+
+public:
+	~UTypedElementTransformHeadsUpWidgetFactory() override = default;
+
+	void RegisterQueries(ITypedElementDataStorageInterface& DataStorage) const override;
+	void RegisterWidgetConstructors(ITypedElementDataStorageInterface& DataStorage,
+		ITypedElementDataStorageUiInterface& DataStorageUi) const override;
+};
+
+USTRUCT()
+struct TYPEDELEMENTSDATASTORAGEUI_API FTypedElementTransformHeadsUpWidgetConstructor : public FTypedElementWidgetConstructor
+{
+	GENERATED_BODY()
+
+public:
+	FTypedElementTransformHeadsUpWidgetConstructor();
+	~FTypedElementTransformHeadsUpWidgetConstructor() override = default;
+
+	TConstArrayView<const UScriptStruct*> GetAdditionalColumnsList() const override;
+	bool CanBeReused() const override;
+
+protected:
+	TSharedPtr<SWidget> CreateWidget() override;
+	bool FinalizeWidget(
+		ITypedElementDataStorageInterface* DataStorage,
+		ITypedElementDataStorageUiInterface* DataStorageUi,
+		TypedElementRowHandle Row,
+		const TSharedPtr<SWidget>& Widget) override;
+};
+
+USTRUCT(meta = (DisplayName = "Heads up display for transforms widget"))
+struct TYPEDELEMENTSDATASTORAGEUI_API FTypedElementTransformHeadsUpWidgetTag : public FTypedElementDataStorageTag
+{
+	GENERATED_BODY()
+};
