Generalized out retrieving the row for UObjects based objects from TEDS.

This is a first step towards supporting UObjects in TEDS. This provides basic support to explicitly register UObjects and FStructs with TEDS and be able to retrieve the rows from their pointers.

diff --git a/Engine/Plugins/Experimental/TypedElementsDataStorage/Source/TypedElementsDataStorage/Private/TypedElementDatabaseCompatibility.cpp b/Engine/Plugins/Experimental/TypedElementsDataStorage/Source/TypedElementsDataStorage/Private/TypedElementDatabaseCompatibility.cpp
--- a/Engine/Plugins/Experimental/TypedElementsDataStorage/Source/TypedElementsDataStorage/Private/TypedElementDatabaseCompatibility.cpp
+++ b/Engine/Plugins/Experimental/TypedElementsDataStorage/Source/TypedElementsDataStorage/Private/TypedElementDatabaseCompatibility.cpp
@@ -4,10 +4,12 @@
 
 #include "Algo/Unique.h"
 #include "Editor.h"
+#include "Elements/Columns/TypedElementCompatibilityColumns.h"
 #include "Elements/Columns/TypedElementLabelColumns.h"
 #include "Elements/Columns/TypedElementMiscColumns.h"
 #include "Elements/Columns/TypedElementPackageColumns.h"
 #include "Elements/Columns/TypedElementTransformColumns.h"
+#include "Elements/Columns/TypedElementTypeInfoColumns.h"
 #include "MassActorEditorSubsystem.h"
 #include "MassActorSubsystem.h"
 #include "TypedElementDataStorageProfilingMacros.h"
@@ -31,15 +33,105 @@ void UTypedElementDatabaseCompatibility::Deinitialize()
 	Reset();
 }
 
-void UTypedElementDatabaseCompatibility::AddCompatibleObject(AActor* Actor)
+TypedElementRowHandle UTypedElementDatabaseCompatibility::AddCompatibleObjectExplicit(UObject* Object)
 {
-	// Registration is delayed for two reasons:
-	//	1. Allows entity creation in a single batch rather than multiple individual additions.
-	//	2. Provides an opportunity to filter out the actors that are created within MASS itself as those will already be registered.
-	ActorsPendingRegistration.Add(Actor);
+	return AddCompatibleObjectExplicit(Object, StandardUObjectTable);
 }
 
-void UTypedElementDatabaseCompatibility::RemoveCompatibleObject(AActor* Actor)
+TypedElementRowHandle UTypedElementDatabaseCompatibility::AddCompatibleObjectExplicit(UObject* Object, TypedElementTableHandle Table)
+{
+	if (Object->IsA<AActor>())
+	{
+		return AddCompatibleObjectExplicit(static_cast<AActor*>(Object));
+	}
+	else
+	{
+		if (ensureMsgf(Storage, TEXT("Trying to add a UObject to Typed Element's Data Storage before the storage is available.")))
+		{
+			TypedElementRowHandle ReservedRow = Storage->ReserveRow();
+			ReverseObjectLookup.Add(Object, ReservedRow);
+			
+			PendingRegistration<TWeakObjectPtr<UObject>>& Pending = UObjectsPendingRegistration.FindOrAdd(Table);
+			Pending.Add(ReservedRow, Object);
+			
+			return ReservedRow;
+		}
+		else
+		{
+			return TypedElementInvalidRowHandle;
+		}
+	}
+}
+
+TypedElementRowHandle UTypedElementDatabaseCompatibility::AddCompatibleObjectExplicit(AActor* Actor)
+{
+	return AddCompatibleObjectExplicit(Actor, StandardActorTable);
+}
+
+TypedElementRowHandle UTypedElementDatabaseCompatibility::AddCompatibleObjectExplicit(AActor* Actor, TypedElementTableHandle Table)
+{
+	if (ensureMsgf(Storage, TEXT("Trying to add an actor to Typed Element's Data Storage before the storage is available.")))
+	{
+		// Registration is delayed for two reasons:
+		//	1. Allows entity creation in a single batch rather than multiple individual additions.
+		//	2. Provides an opportunity to filter out the actors that are created within MASS itself as those will already be registered.
+		
+		TypedElementRowHandle ReservedRow = Storage->ReserveRow();
+		PendingRegistration<TWeakObjectPtr<AActor>>& Pending = ActorsPendingRegistration.FindOrAdd(Table);
+		Pending.Add(ReservedRow, Actor);
+		return ReservedRow;
+	}
+	return TypedElementInvalidRowHandle;
+}
+
+TypedElementRowHandle UTypedElementDatabaseCompatibility::AddCompatibleObjectExplicit(void* Object, TWeakObjectPtr<const UScriptStruct> TypeInfo)
+{
+	return AddCompatibleObjectExplicit(Object, MoveTemp(TypeInfo), StandardExternalObjectTable);
+}
+
+TypedElementRowHandle UTypedElementDatabaseCompatibility::AddCompatibleObjectExplicit(
+	void* Object, TWeakObjectPtr<const UScriptStruct> TypeInfo, TypedElementTableHandle Table)
+{
+	if (ensureMsgf(Storage, TEXT("Trying to add an object to Typed Element's Data Storage before the storage is available.")))
+	{
+		TypedElementRowHandle ReservedRow = Storage->ReserveRow();
+		ReverseObjectLookup.Add(Object, ReservedRow);
+		PendingRegistration<ExternalObjectRegistration>& Pending = ExternalObjectsPendingRegistration.FindOrAdd(Table);
+		Pending.Add(ReservedRow, ExternalObjectRegistration{ .Object = Object, .TypeInfo = TypeInfo });
+		return ReservedRow;
+	}
+	else
+	{
+		return TypedElementInvalidRowHandle;
+	}
+}
+
+void UTypedElementDatabaseCompatibility::RemoveCompatibleObjectExplicit(UObject* Object)
+{
+	if (Object->IsA<AActor>())
+	{
+		RemoveCompatibleObjectExplicit(static_cast<AActor*>(Object));
+	}
+	else
+	{
+		checkf(Storage, TEXT("Removing compatible objects is not supported before Typed Element's Database compatibility manager has been initialized."));
+		if (TypedElementRowHandle* Row = ReverseObjectLookup.Find(Object))
+		{
+			Storage->RemoveRow(*Row);
+		}
+	}
+}
+
+void UTypedElementDatabaseCompatibility::RemoveCompatibleObjectExplicit(void* Object)
+{
+	checkf(Storage, TEXT("Removing compatible objects is not supported before Typed Element's Database compatibility manager has been initialized."));
+	if (TypedElementRowHandle* Row = ReverseObjectLookup.Find(Object))
+	{
+		Storage->RemoveRow(*Row);
+	}
+}
+
+void UTypedElementDatabaseCompatibility::RemoveCompatibleObjectExplicit(AActor* Actor)
 {
 	checkf(Storage, TEXT("Removing compatible objects is not supported before Typed Element's Database compatibility manager has been initialized."));
 
@@ -64,7 +156,26 @@ void UTypedElementDatabaseCompatibility::RemoveCompatibleObject(AActor* Actor)
 	}
 }
 
-TypedElementRowHandle UTypedElementDatabaseCompatibility::FindRowWithCompatibleObject(const TObjectKey<const AActor> Actor) const
+TypedElementRowHandle UTypedElementDatabaseCompatibility::FindRowWithCompatibleObjectExplicit(const UObject* Object) const
+{
+	if (Object && Storage && Storage->IsAvailable())
+	{
+		const AActor* Actor = Cast<AActor>(Object);
+		if (Actor && ActorSubsystem)
+		{
+			FMassEntityHandle Entity = ActorSubsystem->GetEntityHandleFromActor(Actor);
+			return Entity.IsValid() ? Entity.AsNumber() : TypedElementInvalidRowHandle;
+		}
+		else
+		{
+			const TypedElementRowHandle* Row = ReverseObjectLookup.Find(Object);
+			return Row ? *Row : TypedElementInvalidRowHandle;
+		}
+	}
+	return TypedElementInvalidRowHandle;
+}
+
+TypedElementRowHandle UTypedElementDatabaseCompatibility::FindRowWithCompatibleObjectExplicit(const AActor* Actor) const
 {
 	if (Storage && ActorSubsystem && Storage->IsAvailable())
 	{
@@ -77,6 +188,16 @@ TypedElementRowHandle UTypedElementDatabaseCompatibility::FindRowWithCompatibleO
 	}
 }
 
+TypedElementRowHandle UTypedElementDatabaseCompatibility::FindRowWithCompatibleObjectExplicit(const void* Object) const
+{
+	if (Object && Storage && Storage->IsAvailable())
+	{
+		const TypedElementRowHandle* Row = ReverseObjectLookup.Find(Object);
+		return Row ? *Row : TypedElementInvalidRowHandle;
+	}
+	return TypedElementInvalidRowHandle;
+}
+
 void UTypedElementDatabaseCompatibility::Prepare()
 {
 	UMassActorEditorSubsystem* MassActorEditorSubsystem = Storage->GetExternalSystem<UMassActorEditorSubsystem>();
@@ -93,21 +214,26 @@ void UTypedElementDatabaseCompatibility::Reset()
 
 void UTypedElementDatabaseCompatibility::CreateStandardArchetypes()
 {
-	StandardActorTable = Storage->RegisterTable(MakeArrayView(
-		{
-			FMassActorFragment::StaticStruct(),
-			FTypedElementLabelColumn::StaticStruct(),
-			FTypedElementLabelHashColumn::StaticStruct(),
-			FTypedElementPackagePathColumn::StaticStruct(),
-			FTypedElementPackageLoadedPathColumn::StaticStruct(),
-			FTypedElementSyncFromWorldTag::StaticStruct()
-		}), FName("Editor_StandardActorTable"));
+	StandardActorTable = Storage->RegisterTable(TTypedElementColumnTypeList<
+			FMassActorFragment, FTypedElementClassTypeInfoColumn,
+			FTypedElementLabelColumn, FTypedElementLabelHashColumn,
+			FTypedElementPackagePathColumn, FTypedElementPackageLoadedPathColumn,
+			FTypedElementSyncFromWorldTag>(), 
+		FName("Editor_StandardActorTable"));
 
 	StandardActorWithTransformTable = Storage->RegisterTable(StandardActorTable,
-		MakeArrayView(
-		{
-			FTypedElementLocalTransformColumn::StaticStruct()
-		}), FName("Editor_StandardActorWithTransformTable"));
+		TTypedElementColumnTypeList<FTypedElementLocalTransformColumn>(),
+		FName("Editor_StandardActorWithTransformTable"));
+
+	StandardUObjectTable = Storage->RegisterTable(TTypedElementColumnTypeList<
+			FTypedElementUObjectColumn, FTypedElementClassTypeInfoColumn,
+			FTypedElementSyncFromWorldTag>(), 
+		FName("Editor_StandardUObjectTable"));
+
+	StandardExternalObjectTable = Storage->RegisterTable(TTypedElementColumnTypeList<
+			FTypedElementExternalObjectColumn, FTypedElementScriptStructTypeInfoColumn,
+			FTypedElementSyncFromWorldTag>(), 
+		FName("Editor_StandardExternalObjectTable"));
 }
 
 void UTypedElementDatabaseCompatibility::Tick()
@@ -115,107 +241,258 @@ void UTypedElementDatabaseCompatibility::Tick()
 	TEDS_EVENT_SCOPE(TEXT("Compatibility Tick"))
 	UWorld* EditorWorld = GEditor->GetEditorWorldContext().World();
 
-	// Delay processing until the required systems are available by not clearing the pending actor list.
+	// Delay processing until the required systems are available by not clearing any lists or doing any work.
 	if (Storage && Storage->IsAvailable() && EditorWorld)
 	{
-		if (ActorsPendingRegistration.Num() > 0)
+		TickPendingActorRegistration(EditorWorld);
+		TickPendingUObjectRegistration();
+		TickPendingExternalObjectRegistration();
+		TickActorSync();
+	}
+}
+
+template<typename AddressType>
+void UTypedElementDatabaseCompatibility::PendingRegistration<AddressType>::Add(TypedElementRowHandle ReservedRowHandle, AddressType Address)
+{
+	Addresses.Add(Forward<AddressType>(Address));
+	ReservedRowHandles.Add(ReservedRowHandle);
+}
+
+template<typename AddressType>
+bool UTypedElementDatabaseCompatibility::PendingRegistration<AddressType>::IsEmpty() const
+{
+	// ReservedRowHandles can also be returned as they'll both have the same length.
+	return Addresses.IsEmpty();
+}
+
+template<typename AddressType>
+int32 UTypedElementDatabaseCompatibility::PendingRegistration<AddressType>::Num() const
+{
+	// ReservedRowHandles can also be returned as they'll both have the same length.
+	return Addresses.Num();
+}
+
+template<typename AddressType>
+TArrayView<AddressType> UTypedElementDatabaseCompatibility::PendingRegistration<AddressType>::GetAddresses()
+{
+	return Addresses;
+}
+
+template<typename AddressType>
+TArrayView<TypedElementRowHandle> UTypedElementDatabaseCompatibility::PendingRegistration<AddressType>::GetReservedRowHandles()
+{
+	return ReservedRowHandles;
+}
+
+template<typename AddressType>
+void UTypedElementDatabaseCompatibility::PendingRegistration<AddressType>::RemoveInvalidEntries(
+	ITypedElementDataStorageInterface& StorageInterface, const TFunctionRef<bool(const AddressType&)>& Validator)
+{
+	checkf(Addresses.Num() == ReservedRowHandles.Num(),
+		TEXT("The reserved row handle count (%i) didn't match the stored pointer count (%i)."),
+		ReservedRowHandles.Num(), Addresses.Num());
+
+	AddressType* AddressBegin = Addresses.GetData();
+	AddressType* AddressIt = Addresses.GetData();
+	AddressType* AddressEnd = AddressBegin + Addresses.Num();
+	TypedElementRowHandle* RowHandleBegin = ReservedRowHandles.GetData();
+	TypedElementRowHandle* RowHandleIt = ReservedRowHandles.GetData();
+	while (AddressIt != AddressEnd)
+	{
+		if (Validator(*AddressIt))
+		{
+			++AddressIt;
+			++RowHandleIt;
+		}
+		else
 		{
-			// Filter out the actors that are already registered or already destroyed. 
-			// The most common case for this is actors created from within MASS.
-			TWeakObjectPtr<AActor>* ActorBegin = ActorsPendingRegistration.GetData();
-			TWeakObjectPtr<AActor>* ActorIt = ActorsPendingRegistration.GetData();
-			TWeakObjectPtr<AActor>* ActorEnd = ActorBegin + ActorsPendingRegistration.Num();
-			while (ActorIt != ActorEnd)
+			// Don't shrink the registration array as the array will be reused with a variety of different actor counts.
+			// If memory size becomes an issue it's better to resize the array once after this loop rather than within the
+			// loop to avoid many resizes happening.
+			constexpr bool bAllowToShrink = false;
+			Addresses.RemoveAtSwap(AddressIt - AddressBegin, 1, bAllowToShrink);
+			StorageInterface.RemoveRow(*RowHandleIt);
+			ReservedRowHandles.RemoveAtSwap(RowHandleIt - RowHandleBegin, 1, bAllowToShrink);
+			--AddressEnd;
+		}
+	}
+}
+
+template<typename AddressType>
+void UTypedElementDatabaseCompatibility::PendingRegistration<AddressType>::ProcessEntries(ITypedElementDataStorageInterface& StorageInterface,
+	TypedElementTableHandle Table, const TFunctionRef<void(TypedElementRowHandle, const AddressType&)>& SetupRowCallback)
+{
+	if (!IsEmpty())
+	{
+		AddressType* TargetIt = GetAddresses().GetData();
+		AddressType* TargetEnd = TargetIt + Num();
+		StorageInterface.BatchAddRow(Table, GetReservedRowHandles(), 
+			[this, &SetupRowCallback, &TargetIt, TargetEnd](TypedElementRowHandle Row)
 			{
-				if (
-					(*ActorIt).IsValid() && // If not valid or stale then the actor has already been deleted before it could be registered.
-					!ActorSubsystem->GetEntityHandleFromActor((*ActorIt).Get()).IsValid() && // If true, this actor is already registered with the Data Storage.
-					(*ActorIt)->GetWorld() == EditorWorld // Only record actors that are added to the editor to avoid taking PIE into account.
-					) 
+				SetupRowCallback(Row, *TargetIt);
+				checkf(TargetIt < TargetEnd, TEXT("More (%i) entities were added than were requested (%i)."), TargetEnd - TargetIt, Num());
+				++TargetIt;
+			});
+	}
+}
+
+template<typename AddressType>
+void UTypedElementDatabaseCompatibility::PendingRegistration<AddressType>::Reset()
+{
+	Addresses.Reset();
+	ReservedRowHandles.Reset();
+}
+
+void UTypedElementDatabaseCompatibility::TickPendingActorRegistration(UWorld* EditorWorld)
+{
+	if (!ActorsPendingRegistration.IsEmpty())
+	{
+		// Filter out the actors that are already registered or already destroyed. 
+		// The most common case for this is actors created from within MASS.
+		for (auto It = ActorsPendingRegistration.CreateIterator(); It; ++It)
+		{
+			It->Value.RemoveInvalidEntries(*Storage,
+				[this, EditorWorld](const TWeakObjectPtr<AActor>& Actor)
 				{
-					++ActorIt;
-				}
-				else
+					AActor* Instance = Actor.Get();
+					return Instance && Instance->GetWorld() == EditorWorld && !ActorSubsystem->GetEntityHandleFromActor(Instance).IsValid();
+				});
+		}
+
+		// Add the remaining actors to the data storage.
+		for (auto It = ActorsPendingRegistration.CreateIterator(); It; ++It)
+		{
+			It->Value.ProcessEntries(*Storage, It->Key,
+				[this](TypedElementRowHandle Row, const TWeakObjectPtr<AActor>& ActorPtr)
 				{
-					// Don't shrink the registration array as the array will be reused with a variety of different actor counts.
-					// If memory size becomes an issue it's better to resize the array once after this loop rather than within the
-					// loop to avoid many resizes happening.
-					constexpr bool bAllowToShrink = false;
-					ActorsPendingRegistration.RemoveAtSwap(ActorIt - ActorBegin, 1, bAllowToShrink);
-					--ActorEnd;
-				}
-			}
+					FMassActorFragment* ActorStore = Storage->AddOrGetColumn<FMassActorFragment>(Row);
+					checkf(ActorStore, TEXT("Failed to retrieve or add FMassActorFragment to newly created row."));
 
-			if (ActorBegin != ActorEnd) // Check if the number of actors to add isn't zero.
-			{
-				// Add the remaining actors to the data storage.
-				ActorIt = ActorsPendingRegistration.GetData();
-				Storage->BatchAddRow(StandardActorTable, ActorsPendingRegistration.Num(), [this, &ActorIt, ActorEnd](TypedElementRowHandle Row)
-					{
-						auto ActorStore = Storage->GetColumn<FMassActorFragment>(Row);
-						checkf(ActorStore, TEXT("Newly created row didn't contain the expected FMassActorFragment."));
-						
-						constexpr bool bIsOwnedByMass = false;
-					
-						AActor* Actor = ActorIt->Get();
-						check(Actor);
-						ActorStore->SetNoHandleMapUpdate(FMassEntityHandle::FromNumber(Row), Actor, bIsOwnedByMass);
-						ActorSubsystem->SetHandleForActor(ActorIt->Get(), FMassEntityHandle::FromNumber(Row));
-					
-						ActorsNeedingFullSync.Emplace(Actor);
-
-						checkf(ActorIt < ActorEnd,
-							TEXT("More (%i) entities were added than were requested (%i)."), ActorEnd - ActorIt, ActorsPendingRegistration.Num());
-						++ActorIt;
-					});
-
-				// Reset the container for next set of actors.
-				ActorsPendingRegistration.Reset();
-			}
+					constexpr bool bIsOwnedByMass = false;
+
+					AActor* Actor = ActorPtr.Get();
+					check(Actor);
+					ActorStore->SetNoHandleMapUpdate(FMassEntityHandle::FromNumber(Row), Actor, bIsOwnedByMass);
+					ActorSubsystem->SetHandleForActor(Actor, FMassEntityHandle::FromNumber(Row));
+
+					Storage->AddOrGetColumn<FTypedElementClassTypeInfoColumn>(Row, FTypedElementClassTypeInfoColumn{ .TypeInfo = Actor->GetClass() });
+					checkf(ActorStore, TEXT("Failed to retrieve or add FTypedElementClassTypeInfoColumn to newly created row."));
+
+					// Make sure the new row is tagged for update.
+					Storage->AddColumn<FTypedElementSyncFromWorldTag>(Row);
+				});
+		}
+			
+		ActorsPendingRegistration.Reset();
+	}
+}
+
+void UTypedElementDatabaseCompatibility::TickPendingUObjectRegistration()
+{
+	if (!UObjectsPendingRegistration.IsEmpty())
+	{
+		// Filter out the objects that are already registered or already destroyed. 
+		for (auto It = UObjectsPendingRegistration.CreateIterator(); It; ++It)
+		{
+			It->Value.RemoveInvalidEntries(*Storage,
+				[this](const TWeakObjectPtr<UObject>& Object)
+				{
+					UObject* Instance = Object.Get();
+					return Instance && FindRowWithCompatibleObjectExplicit(Instance) != TypedElementInvalidRowHandle;
+				});
 		}
 
-		if (!ActorsNeedingFullSync.IsEmpty())
+		// Add the remaining object to the data storage.
+		for (auto It = UObjectsPendingRegistration.CreateIterator(); It; ++It)
 		{
-			TEDS_EVENT_SCOPE(TEXT("Process ActorsNeedingFullSync"));
-			// Deduplicate to avoid duplicate reverse lookups or adding tags more than once
-			{
-				TEDS_EVENT_SCOPE(TEXT("Deduplicate ActorsNeedingFullSync"));
-				ActorsNeedingFullSync.Sort();
-				ActorsNeedingFullSync.SetNum(Algo::Unique(ActorsNeedingFullSync));
-			}
+			It->Value.ProcessEntries(*Storage, It->Key, 
+				[this](TypedElementRowHandle Row, const TWeakObjectPtr<UObject>& Object)
+				{
+					Storage->AddOrGetColumn<FTypedElementUObjectColumn>(Row, FTypedElementUObjectColumn{ .Object = Object });
+					Storage->AddOrGetColumn<FTypedElementClassTypeInfoColumn>(Row, FTypedElementClassTypeInfoColumn{ .TypeInfo = Object->GetClass() });
+					// Make sure the new row is tagged for update.
+					Storage->AddColumn<FTypedElementSyncFromWorldTag>(Row);
+				});
+		}
+
+		UObjectsPendingRegistration.Reset();
+	}
+}
+
+void UTypedElementDatabaseCompatibility::TickPendingExternalObjectRegistration()
+{
+	if (!ExternalObjectsPendingRegistration.IsEmpty())
+	{
+		// Filter out the objects that are already registered or already destroyed. 
+		for (auto It = ExternalObjectsPendingRegistration.CreateIterator(); It; ++It)
+		{
+			It->Value.RemoveInvalidEntries(*Storage,
+				[this](const ExternalObjectRegistration& Object)
+				{
+					return Object.Object && FindRowWithCompatibleObjectExplicit(Object.Object) != TypedElementInvalidRowHandle;
+				});
+		}
+
+		// Add the remaining object to the data storage.
+		for (auto It = ExternalObjectsPendingRegistration.CreateIterator(); It; ++It)
+		{
+			It->Value.ProcessEntries(*Storage, It->Key, 
+				[this](TypedElementRowHandle Row, const ExternalObjectRegistration& Object)
+				{
+					Storage->AddOrGetColumn<FTypedElementUObjectColumn>(Row, FTypedElementExternalObjectColumn{ .Object = Object.Object });
+					Storage->AddOrGetColumn<FTypedElementClassTypeInfoColumn>(Row, FTypedElementScriptStructTypeInfoColumn{ .TypeInfo = Object.TypeInfo });
+					// Make sure the new row is tagged for update.
+					Storage->AddColumn<FTypedElementSyncFromWorldTag>(Row);
+				});
+		}
+
+		ExternalObjectsPendingRegistration.Reset();
+	}
+}
+
+void UTypedElementDatabaseCompatibility::TickActorSync()
+{
+	if (!ActorsNeedingFullSync.IsEmpty())
+	{
+		TEDS_EVENT_SCOPE(TEXT("Process ActorsNeedingFullSync"));
+		// Deduplicate to avoid duplicate reverse lookups or adding tags more than once
+		{
+			TEDS_EVENT_SCOPE(TEXT("Deduplicate ActorsNeedingFullSync"));
+			ActorsNeedingFullSync.Sort();
+			ActorsNeedingFullSync.SetNum(Algo::Unique(ActorsNeedingFullSync), /* bAllowShrinking */ false);
+		}
+
+		TArray<TypedElementRowHandle> RowHandles;
+		{
+			TEDS_EVENT_SCOPE(TEXT("Reverse lookup Rows from Actors"));
 
-			TArray<TypedElementRowHandle> RowHandles;
+			RowHandles.SetNumUninitialized(ActorsNeedingFullSync.Num());
 			{
-				TEDS_EVENT_SCOPE(TEXT("Reverse lookup Rows from Actors"));
-				
-				RowHandles.SetNumUninitialized(ActorsNeedingFullSync.Num());
+				int32 RowHandleIndex = 0;
+				for (int32 ActorIndex = 0, End = ActorsNeedingFullSync.Num(); ActorIndex < End; ++ActorIndex)
 				{
-					int32 RowHandleIndex = 0;
-					for (int32 ActorIndex = 0, End = ActorsNeedingFullSync.Num(); ActorIndex < End; ++ActorIndex)
+					const TObjectKey<const AActor> ActorKey = ActorsNeedingFullSync[ActorIndex];
+					const TypedElementRowHandle Row = FindRowWithCompatibleObject(ActorKey);
+					if (Row != TypedElementInvalidRowHandle)
 					{
-						const TObjectKey<const AActor> ActorKey = ActorsNeedingFullSync[ActorIndex];
-						const TypedElementRowHandle Row = FindRowWithCompatibleObject(ActorKey);
-						if (Row != TypedElementInvalidRowHandle)
-						{
-							RowHandles[RowHandleIndex++] = Row;
-						}
+						RowHandles[RowHandleIndex++] = Row;
 					}
-					const int32 RowHandleCount = RowHandleIndex;
-					const bool bAllowShrinking = false;
-					RowHandles.SetNum(RowHandleCount, bAllowShrinking);
 				}
-
-				ActorsNeedingFullSync.Reset();
+				const int32 RowHandleCount = RowHandleIndex;
+				const bool bAllowShrinking = false;
+				RowHandles.SetNum(RowHandleCount, bAllowShrinking);
 			}
 
+			ActorsNeedingFullSync.Reset();
+		}
+
+		{
+			TEDS_EVENT_SCOPE(TEXT("Add SyncFromWorld Tag"));
+			// Tag the rows containing actor data that they need to be synced
+			// Note: Watch out for the performance of this, may end up doing a lot of row moves
+			for (TypedElementRowHandle Row : RowHandles)
 			{
-				TEDS_EVENT_SCOPE(TEXT("Add SyncFromWorld Tag"));
-				// Tag the rows containing actor data that they need to be synced
-				// Note: Watch out for the performance of this, may end up doing a lot of row moves
-				for (TypedElementRowHandle Row : RowHandles)
-				{
-					Storage->AddColumn<FTypedElementSyncFromWorldTag>(Row);
-				}
+				Storage->AddColumn<FTypedElementSyncFromWorldTag>(Row);
 			}
 		}
 	}
diff --git a/Engine/Plugins/Experimental/TypedElementsDataStorage/Source/TypedElementsDataStorage/Private/TypedElementDatabaseCompatibility.h b/Engine/Plugins/Experimental/TypedElementsDataStorage/Source/TypedElementsDataStorage/Private/TypedElementDatabaseCompatibility.h
--- a/Engine/Plugins/Experimental/TypedElementsDataStorage/Source/TypedElementsDataStorage/Private/TypedElementDatabaseCompatibility.h
+++ b/Engine/Plugins/Experimental/TypedElementsDataStorage/Source/TypedElementsDataStorage/Private/TypedElementDatabaseCompatibility.h
@@ -26,9 +26,20 @@ public:
 	void Initialize(ITypedElementDataStorageInterface* StorageInterface);
 	void Deinitialize();
 
-	void AddCompatibleObject(AActor* Actor) override;
-	void RemoveCompatibleObject(AActor* Actor) override;
-	TypedElementRowHandle FindRowWithCompatibleObject(const TObjectKey<const AActor> Actor) const override;
+	TypedElementRowHandle AddCompatibleObjectExplicit(UObject* Object) override;
+	TypedElementRowHandle AddCompatibleObjectExplicit(UObject* Object, TypedElementTableHandle Table) override;
+	TypedElementRowHandle AddCompatibleObjectExplicit(AActor* Actor) override;
+	TypedElementRowHandle AddCompatibleObjectExplicit(AActor* Actor, TypedElementTableHandle Table) override;
+	TypedElementRowHandle AddCompatibleObjectExplicit(void* Object, TWeakObjectPtr<const UScriptStruct> TypeInfo);
+	TypedElementRowHandle AddCompatibleObjectExplicit(void* Object, TWeakObjectPtr<const UScriptStruct> TypeInfo, TypedElementTableHandle Table);
+	
+	void RemoveCompatibleObjectExplicit(UObject* Object) override;
+	void RemoveCompatibleObjectExplicit(AActor* Actor) override;
+	void RemoveCompatibleObjectExplicit(void* Object) override;
+
+	TypedElementRowHandle FindRowWithCompatibleObjectExplicit(const UObject* Object) const override;
+	TypedElementRowHandle FindRowWithCompatibleObjectExplicit(const AActor* Actor) const override;
+	TypedElementRowHandle FindRowWithCompatibleObjectExplicit(const void* Object) const override;
 
 private:
 	void Prepare();
@@ -36,20 +47,54 @@ private:
 	void CreateStandardArchetypes();
 	
 	void Tick();
+	void TickPendingActorRegistration(UWorld* EditorWorld);
+	void TickPendingUObjectRegistration();
+	void TickPendingExternalObjectRegistration();
+	void TickActorSync();
 
 	void OnPostEditChangeProperty(UObject* Object, FPropertyChangedEvent& PropertyChangedEvent);
 	
-	TArray<TWeakObjectPtr<AActor>> ActorsPendingRegistration;
+	template<typename AddressType>
+	struct PendingRegistration
+	{
+	private:
+		TArray<AddressType> Addresses;
+		TArray<TypedElementRowHandle> ReservedRowHandles;
+
+	public:
+		void Add(TypedElementRowHandle ReservedRowHandle, AddressType Address);
+		bool IsEmpty() const;
+		int32 Num() const;
+		TArrayView<AddressType> GetAddresses();
+		TArrayView<TypedElementRowHandle> GetReservedRowHandles();
+
+		void RemoveInvalidEntries(ITypedElementDataStorageInterface& Storage, const TFunctionRef<bool(const AddressType&)>& Validator);
+		void ProcessEntries(ITypedElementDataStorageInterface& Storage, TypedElementTableHandle Table,
+			const TFunctionRef<void(TypedElementRowHandle, const AddressType&)>& SetupRowCallback);
+		void Reset();
+	};
+	struct ExternalObjectRegistration
+	{
+		void* Object;
+		TWeakObjectPtr<const UScriptStruct> TypeInfo;
+	};
+	TMap<TypedElementTableHandle, PendingRegistration<TWeakObjectPtr<AActor>>> ActorsPendingRegistration;
+	TMap<TypedElementTableHandle, PendingRegistration<TWeakObjectPtr<UObject>>> UObjectsPendingRegistration;
+	TMap<TypedElementTableHandle, PendingRegistration<ExternalObjectRegistration>> ExternalObjectsPendingRegistration;
 	
 	TypedElementTableHandle StandardActorTable{ TypedElementInvalidTableHandle };
 	TypedElementTableHandle StandardActorWithTransformTable{ TypedElementInvalidTableHandle };
+	TypedElementTableHandle StandardUObjectTable{ TypedElementInvalidTableHandle };
+	TypedElementTableHandle StandardExternalObjectTable{ TypedElementInvalidTableHandle };
 	ITypedElementDataStorageInterface* Storage{ nullptr };
 	TSharedPtr<FMassActorManager> ActorSubsystem;
 
-	/**~
-	 * Reference of actors that need to be fully synced from the world to the database
+	TMap<void*, TypedElementRowHandle> ReverseObjectLookup;
+
+	/**
+	 * Reference of actors that need to be fully synced from the world to the database.
 	 * May have duplicates
-	 * Caution: Could point to actors that have been GCd
+	 * Caution: Could point to actors that have been GC-ed
 	 */
 	TArray<TObjectKey<const AActor>> ActorsNeedingFullSync;
 
diff --git a/Engine/Source/Runtime/TypedElementFramework/Public/Elements/Columns/TypedElementCompatibilityColumns.h b/Engine/Source/Runtime/TypedElementFramework/Public/Elements/Columns/TypedElementCompatibilityColumns.h
new file mode 100644
--- /dev/null
+++ b/Engine/Source/Runtime/TypedElementFramework/Public/Elements/Columns/TypedElementCompatibilityColumns.h
@@ -0,0 +1,35 @@
+// Copyright Epic Games, Inc. All Rights Reserved.
+
+#pragma once
+
+#include "Elements/Interfaces/TypedElementDataStorageInterface.h"
+#include "UObject/ObjectMacros.h"
+#include "UObject/WeakObjectPtrTemplates.h"
+
+#include "TypedElementCompatibilityColumns.generated.h"
+
+class UObject;
+
+/**
+ * Column containing a non-owning reference to a UObject.
+ */
+USTRUCT(meta = (DisplayName = "UObject reference"))
+struct FTypedElementUObjectColumn final : public FTypedElementDataStorageColumn
+{
+	GENERATED_BODY()
+
+	TWeakObjectPtr<UObject> Object;
+};
+
+/**
+ * Column containing a non-owning reference to an arbitrary object. It's strongly recommended
+ * to also add a FTypedElementScriptStructTypeInfoColumn to make sure the type can be safely
+ * recovered.
+ */
+USTRUCT(meta = (DisplayName = "External object reference"))
+struct FTypedElementExternalObjectColumn final : public FTypedElementDataStorageColumn
+{
+	GENERATED_BODY()
+
+	void* Object;
+};
\ No newline at end of file
diff --git a/Engine/Source/Runtime/TypedElementFramework/Public/Elements/Interfaces/TypedElementDataStorageCompatibilityInterface.h b/Engine/Source/Runtime/TypedElementFramework/Public/Elements/Interfaces/TypedElementDataStorageCompatibilityInterface.h
--- a/Engine/Source/Runtime/TypedElementFramework/Public/Elements/Interfaces/TypedElementDataStorageCompatibilityInterface.h
+++ b/Engine/Source/Runtime/TypedElementFramework/Public/Elements/Interfaces/TypedElementDataStorageCompatibilityInterface.h
@@ -6,6 +6,9 @@
 #include "UObject/Interface.h"
 #include "UObject/ObjectKey.h"
 #include "UObject/ObjectMacros.h"
+#include "UObject/ObjectPtr.h"
+#include "UObject/StrongObjectPtr.h"
+#include "UObject/WeakObjectPtrTemplates.h"
 
 #include "TypedElementDataStorageCompatibilityInterface.generated.h"
 
@@ -26,11 +29,118 @@ class TYPEDELEMENTFRAMEWORK_API ITypedElementDataStorageCompatibilityInterface
 	GENERATED_BODY()
 
 public:
-	virtual void AddCompatibleObject(AActor* Actor) = 0;
-	virtual void RemoveCompatibleObject(AActor* Actor) = 0;
-	virtual TypedElementRowHandle FindRowWithCompatibleObject(const TObjectKey<const AActor> Actor) const = 0;
+	/**
+	 * @section Type-agnostic functions
+	 * These allow compatibility with any type. These do eventually fall back to the explicit versions.
+	 * Any references given are non-owning so it's up to the caller to deregister the object after it's no longer
+	 * available.
+	 */
+
+	/** 
+	 * Adds a reference to an existing object to the data storage. The data storage does NOT take ownership of the object and
+	 * the caller is responsible for managing the life cycle of the object. The address is only used for associating the object
+	 * with a row and to setup the initial row data.
+	 */
+	template<typename ObjectType>
+	TypedElementRowHandle AddCompatibleObject(ObjectType Object);
+	template<typename ObjectType>
+	TypedElementRowHandle AddCompatibleObject(ObjectType Object, TypedElementTableHandle Table);
+
+	/** Removes a previously registered object from the data storage. */
+	template<typename ObjectType>
+	void RemoveCompatibleObject(ObjectType Object);
+
+	template<typename ObjectType>
+	TypedElementRowHandle FindRowWithCompatibleObject(const ObjectType Object) const;
+
+
+	/**
+	 * @section Explicit functions
+	 * These are functions that work on specific types.
+	 */
+
+	/** Adds a UObject to the data storage. */
+	virtual TypedElementRowHandle AddCompatibleObjectExplicit(UObject* Object) = 0;
+	virtual TypedElementRowHandle AddCompatibleObjectExplicit(UObject* Object, TypedElementTableHandle Table) = 0;
+	/** Adds an actor to the data storage. */
+	virtual TypedElementRowHandle AddCompatibleObjectExplicit(AActor* Actor) = 0;
+	virtual TypedElementRowHandle AddCompatibleObjectExplicit(AActor* Actor, TypedElementTableHandle Table) = 0;
+	/** Adds an FStruct to the data storage. */
+	virtual TypedElementRowHandle AddCompatibleObjectExplicit(void* Object, TWeakObjectPtr<const UScriptStruct> TypeInfo) = 0;
+	virtual TypedElementRowHandle AddCompatibleObjectExplicit(void* Object, TWeakObjectPtr<const UScriptStruct> TypeInfo, TypedElementTableHandle Table) = 0;
+
+	/** Removes a UObject from the data storage. */
+	virtual void RemoveCompatibleObjectExplicit(UObject* Object) = 0;
+	/** Removes an actor from the data storage. */
+	virtual void RemoveCompatibleObjectExplicit(AActor* Actor) = 0;
+	/** Removes an FStruct from the data storage. */
+	virtual void RemoveCompatibleObjectExplicit(void* Object) = 0;
+
+	/** Finds a previously stored UObject. If not found an invalid row handle will be returned. */
+	virtual TypedElementRowHandle FindRowWithCompatibleObjectExplicit(const UObject* Object) const = 0;
+	/** 
+	 * Finds a previously stored actor based on the object key. If not found an invalid row handle will be returned.
+	 * While FindRowWithCompatibleObject(const TObjectKey<const UObject> Object) can also be used, this call will be slightly
+	 * faster if it's already known that the target is an actor.
+	 */
+	virtual TypedElementRowHandle FindRowWithCompatibleObjectExplicit(const AActor* Actor) const = 0;
+	/** Finds a previously stored FStructe . If not found an invalid row handle will be returned. */
+	virtual TypedElementRowHandle FindRowWithCompatibleObjectExplicit(const void* Object) const = 0;
 };
 
+template<typename Type> Type* GetRawPointer(const TWeakObjectPtr<Type> Object)	{ return Object.Get(); }
+template<typename Type> Type* GetRawPointer(const TObjectPtr<Type> Object)		{ return Object.Get(); }
+template<typename Type> Type* GetRawPointer(const TStrongObjectPtr<Type> Object){ return Object.Get(); }
+template<typename Type> Type* GetRawPointer(const TObjectKey<Type> Object)		{ return Object.ResolveObjectPtr(); }
+template<typename Type> Type* GetRawPointer(const TUniquePtr<Type> Object)		{ return Object.Get(); }
+template<typename Type> Type* GetRawPointer(const TSharedPtr<Type> Object)		{ return Object.Get(); }
+template<typename Type> Type* GetRawPointer(Type* Object)						{ return Object; }
+template<typename Type> Type* GetRawPointer(Type& Object)						{ return &Object; }
+
+template<typename ObjectType>
+TypedElementRowHandle ITypedElementDataStorageCompatibilityInterface::AddCompatibleObject(ObjectType Object)
+{
+	auto RawPointer = GetRawPointer(Object);
+	using BaseType = std::remove_cv_t<std::remove_pointer_t<decltype(RawPointer)>>;
+
+	if constexpr (std::is_same_v<BaseType, AActor> || std::is_same_v<BaseType, UObject>)
+	{
+		return AddCompatibleObjectExplicit(RawPointer);
+	}
+	else
+	{
+		return AddCompatibleObjectExplicit(RawPointer, BaseType::StaticStruct());
+	}
+}
+
+template<typename ObjectType>
+TypedElementRowHandle ITypedElementDataStorageCompatibilityInterface::AddCompatibleObject(ObjectType Object, TypedElementTableHandle Table)
+{
+	auto RawPointer = GetRawPointer(Object);
+	using BaseType = std::remove_cv_t<std::remove_pointer_t<decltype(RawPointer)>>;
+
+	if constexpr (std::is_same_v<BaseType, AActor> || std::is_same_v<BaseType, UObject>)
+	{
+		return AddCompatibleObjectExplicit(RawPointer, Table);
+	}
+	else
+	{
+		return AddCompatibleObjectExplicit(RawPointer, BaseType::StaticStruct(), Table);
+	}
+}
+
+template<typename ObjectType>
+void ITypedElementDataStorageCompatibilityInterface::RemoveCompatibleObject(ObjectType Object)
+{
+	RemoveCompatibleObjectExplicit(GetRawPointer(Object));
+}
+
+template<typename ObjectType>
+TypedElementRowHandle ITypedElementDataStorageCompatibilityInterface::FindRowWithCompatibleObject(const ObjectType Object) const
+{
+	return FindRowWithCompatibleObjectExplicit(GetRawPointer(Object));
+}
+
 template<typename Subsystem>
 struct TTypedElementSubsystemTraits final
 {
@@ -85,4 +195,4 @@ struct TTypedElementSubsystemTraits final
 			return false;
 		}
 	}
-};
\ No newline at end of file
+};
