Moved TEDS package integration to work on the more general UObject instead of the more narrow AActor.

diff --git a/Engine/Plugins/Experimental/TypedElementsDataStorage/Source/TypedElementsDataStorage/Private/Compatibility/TypedElementActorPackagePathToColumnQueries.cpp b/Engine/Plugins/Experimental/TypedElementsDataStorage/Source/TypedElementsDataStorage/Private/Compatibility/TypedElementActorPackagePathToColumnQueries.cpp
deleted file mode 100644
diff --git a/Engine/Plugins/Experimental/TypedElementsDataStorage/Source/TypedElementsDataStorage/Private/Compatibility/TypedElementActorPackagePathToColumnQueries.h b/Engine/Plugins/Experimental/TypedElementsDataStorage/Source/TypedElementsDataStorage/Private/Compatibility/TypedElementActorPackagePathToColumnQueries.h
deleted file mode 100644
diff --git a/Engine/Plugins/Experimental/TypedElementsDataStorage/Source/TypedElementsDataStorage/Private/Compatibility/TypedElementUObjectPackagePathToColumnQueries.cpp b/Engine/Plugins/Experimental/TypedElementsDataStorage/Source/TypedElementsDataStorage/Private/Compatibility/TypedElementUObjectPackagePathToColumnQueries.cpp
new file mode 100644
--- /dev/null
+++ b/Engine/Plugins/Experimental/TypedElementsDataStorage/Source/TypedElementsDataStorage/Private/Compatibility/TypedElementUObjectPackagePathToColumnQueries.cpp
@@ -0,0 +1,31 @@
+// Copyright Epic Games, Inc. All Rights Reserved.
+
+#include "Compatibility/TypedElementUObjectPackagePathToColumnQueries.h"
+
+#include "Elements/Columns/TypedElementCompatibilityColumns.h"
+#include "Elements/Columns/TypedElementPackageColumns.h"
+#include "Elements/Framework/TypedElementQueryBuilder.h"
+#include "UObject/Package.h"
+
+void UTypedElementUObjectPackagePathFactory::RegisterQueries(ITypedElementDataStorageInterface& DataStorage) const
+{
+	using namespace TypedElementQueryBuilder;
+	
+	DataStorage.RegisterQuery(
+		Select(
+			TEXT("Sync UObject package info to columns"),
+			FObserver(FObserver::EEvent::Add, FTypedElementUObjectColumn::StaticStruct())
+				.ForceToGameThread(true),
+			[](const FTypedElementUObjectColumn& Object, FTypedElementPackagePathColumn& Path, FTypedElementPackageLoadedPathColumn& LoadedPath)
+			{
+				if (const UObject* ObjectInstance = Object.Object.Get(); ObjectInstance != nullptr)
+				{
+					const UPackage* Target = ObjectInstance->GetPackage();
+					Target->GetPathName(nullptr, Path.Path);
+					LoadedPath.LoadedPath = Target->GetLoadedPath();
+				}
+			}
+		)
+		.Compile()
+	);
+}
diff --git a/Engine/Plugins/Experimental/TypedElementsDataStorage/Source/TypedElementsDataStorage/Private/Compatibility/TypedElementUObjectPackagePathToColumnQueries.h b/Engine/Plugins/Experimental/TypedElementsDataStorage/Source/TypedElementsDataStorage/Private/Compatibility/TypedElementUObjectPackagePathToColumnQueries.h
new file mode 100644
--- /dev/null
+++ b/Engine/Plugins/Experimental/TypedElementsDataStorage/Source/TypedElementsDataStorage/Private/Compatibility/TypedElementUObjectPackagePathToColumnQueries.h
@@ -0,0 +1,19 @@
+// Copyright Epic Games, Inc. All Rights Reserved.
+
+#pragma once
+
+#include "Elements/Interfaces/TypedElementDataStorageFactory.h"
+#include "UObject/ObjectMacros.h"
+
+#include "TypedElementUObjectPackagePathToColumnQueries.generated.h"
+
+UCLASS()
+class UTypedElementUObjectPackagePathFactory : public UTypedElementDataStorageFactory
+{
+	GENERATED_BODY()
+
+public:
+	~UTypedElementUObjectPackagePathFactory() override = default;
+
+	void RegisterQueries(ITypedElementDataStorageInterface& DataStorage) const override;
+};
diff --git a/Engine/Plugins/Experimental/TypedElementsDataStorage/Source/TypedElementsDataStorage/Private/TypedElementDatabaseCompatibility.cpp b/Engine/Plugins/Experimental/TypedElementsDataStorage/Source/TypedElementsDataStorage/Private/TypedElementDatabaseCompatibility.cpp
--- a/Engine/Plugins/Experimental/TypedElementsDataStorage/Source/TypedElementsDataStorage/Private/TypedElementDatabaseCompatibility.cpp
+++ b/Engine/Plugins/Experimental/TypedElementsDataStorage/Source/TypedElementsDataStorage/Private/TypedElementDatabaseCompatibility.cpp
@@ -215,7 +215,7 @@ void UTypedElementDatabaseCompatibility::Reset()
 void UTypedElementDatabaseCompatibility::CreateStandardArchetypes()
 {
 	StandardActorTable = Storage->RegisterTable(TTypedElementColumnTypeList<
-			FMassActorFragment, FTypedElementClassTypeInfoColumn,
+			FMassActorFragment, FTypedElementUObjectColumn, FTypedElementClassTypeInfoColumn,
 			FTypedElementLabelColumn, FTypedElementLabelHashColumn,
 			FTypedElementPackagePathColumn, FTypedElementPackageLoadedPathColumn,
 			FTypedElementSyncFromWorldTag>(), 
@@ -375,9 +375,9 @@ void UTypedElementDatabaseCompatibility::TickPendingActorRegistration(UWorld* Ed
 					ActorStore->SetNoHandleMapUpdate(FMassEntityHandle::FromNumber(Row), Actor, bIsOwnedByMass);
 					ActorSubsystem->SetHandleForActor(Actor, FMassEntityHandle::FromNumber(Row));
 
+					Storage->AddOrGetColumn<FTypedElementUObjectColumn>(Row, FTypedElementUObjectColumn{ .Object = ActorPtr });
 					Storage->AddOrGetColumn<FTypedElementClassTypeInfoColumn>(Row, FTypedElementClassTypeInfoColumn{ .TypeInfo = Actor->GetClass() });
-					checkf(ActorStore, TEXT("Failed to retrieve or add FTypedElementClassTypeInfoColumn to newly created row."));
-
+					
 					// Make sure the new row is tagged for update.
 					Storage->AddColumn<FTypedElementSyncFromWorldTag>(Row);
 				});
