Fixed a string lifetime bug.

When trying to find the longest matching string for column names in TEDS it would occassionaly just return "column" even if a better match existed. This was because the column label that was retrieved was incorrectly stored so went out of scope and as a result the string retrieved from it was also no longer available. Due to timing of memory reuse this issue wouldn't always show up. The text object is now correctly stored, fixing the the invalid string reference retrieved from it.

diff --git a/Engine/Plugins/Experimental/TypedElementsDataStorage/Source/TypedElementsDataStorageUI/Private/TypedElementSceneOutliner.cpp b/Engine/Plugins/Experimental/TypedElementsDataStorage/Source/TypedElementsDataStorageUI/Private/TypedElementSceneOutliner.cpp
--- a/Engine/Plugins/Experimental/TypedElementsDataStorage/Source/TypedElementsDataStorageUI/Private/TypedElementSceneOutliner.cpp
+++ b/Engine/Plugins/Experimental/TypedElementsDataStorage/Source/TypedElementsDataStorageUI/Private/TypedElementSceneOutliner.cpp
@@ -488,13 +488,14 @@ FName FTypedElementSceneOutliner::FindLongestMatchingName(TConstArrayView<TWeakO
 		return FName(ColumnTypes[0]->GetDisplayNameText().ToString());
 	default:
 	{
-		FString LongestMatch(ColumnTypes[0]->GetDisplayNameText().ToString());
+		FString LongestMatch = ColumnTypes[0]->GetDisplayNameText().ToString();
 		const TWeakObjectPtr<const UScriptStruct>* ItEnd = ColumnTypes.end();
 		const TWeakObjectPtr<const UScriptStruct>* It = ColumnTypes.begin();
 		++It; // Skip the first entry as that's already set.
 		for (; It != ItEnd; ++It)
 		{
-			const FString& NextMatch = (*It)->GetDisplayNameText().ToString();
+			FText NextMatchText = (*It)->GetDisplayNameText();
+			const FString& NextMatch = NextMatchText.ToString();
 
 			int32 MatchSize = 0;
 			auto ItLeft = LongestMatch.begin();
