Avoid temporary string creation by using FStringView.

While constructing the name of a column for TEDS' binding with the outliner it was using a temporary string, while this could be avoided by using a FStringView.

diff --git a/Engine/Plugins/Experimental/TypedElementsDataStorage/Source/TypedElementsDataStorageUI/Private/TypedElementSceneOutliner.cpp b/Engine/Plugins/Experimental/TypedElementsDataStorage/Source/TypedElementsDataStorageUI/Private/TypedElementSceneOutliner.cpp
--- a/Engine/Plugins/Experimental/TypedElementsDataStorage/Source/TypedElementsDataStorageUI/Private/TypedElementSceneOutliner.cpp
+++ b/Engine/Plugins/Experimental/TypedElementsDataStorage/Source/TypedElementsDataStorageUI/Private/TypedElementSceneOutliner.cpp
@@ -488,14 +488,15 @@ FName FTypedElementSceneOutliner::FindLongestMatchingName(TConstArrayView<TWeakO
 		return FName(ColumnTypes[0]->GetDisplayNameText().ToString());
 	default:
 	{
-		FString LongestMatch = ColumnTypes[0]->GetDisplayNameText().ToString();
+		FText LongestMatchText = ColumnTypes[0]->GetDisplayNameText();
+		FStringView LongestMatch = LongestMatchText.ToString();
 		const TWeakObjectPtr<const UScriptStruct>* ItEnd = ColumnTypes.end();
 		const TWeakObjectPtr<const UScriptStruct>* It = ColumnTypes.begin();
 		++It; // Skip the first entry as that's already set.
 		for (; It != ItEnd; ++It)
 		{
 			FText NextMatchText = (*It)->GetDisplayNameText();
-			const FString& NextMatch = NextMatchText.ToString();
+			FStringView NextMatch = NextMatchText.ToString();
 
 			int32 MatchSize = 0;
 			auto ItLeft = LongestMatch.begin();
@@ -523,7 +524,7 @@ FName FTypedElementSceneOutliner::FindLongestMatchingName(TConstArrayView<TWeakO
 				return FName(TEXT("Column"), DefaultNameIndex);
 			}
 		}
-		return FName(MoveTemp(LongestMatch));
+		return FName(LongestMatch);
 	}
 	};
 }
