Added support for having TEDS subqueries with a callback.

The following changes were made:
- It's now possible to run a subquery with a callback from a regular query callback. This can be generally applied or limited to a specific row.
- The query context implementations have been updated to use a templated class. This reduced the duplicated boilerplate code originally needed.
- A new query context was added that specifically deals with the case of a single row being called. This allows for several shortcuts and better performance.
- Caching subsystems in direct and subquery callbacks has been removed as subsystems aren't available for access in those contexts.

diff --git a/Engine/Plugins/Experimental/TypedElementsDataStorage/Source/TypedElementsDataStorage/Private/Processors/TypedElementProcessorAdaptors.cpp b/Engine/Plugins/Experimental/TypedElementsDataStorage/Source/TypedElementsDataStorage/Private/Processors/TypedElementProcessorAdaptors.cpp
--- a/Engine/Plugins/Experimental/TypedElementsDataStorage/Source/TypedElementsDataStorage/Private/Processors/TypedElementProcessorAdaptors.cpp
+++ b/Engine/Plugins/Experimental/TypedElementsDataStorage/Source/TypedElementsDataStorage/Private/Processors/TypedElementProcessorAdaptors.cpp
@@ -3,25 +3,41 @@
 #include "Processors/TypedElementProcessorAdaptors.h"
 
 #include <utility>
+#include "Elements/Common/TypedElementQueryTypes.h"
 #include "MassCommonTypes.h"
+#include "MassEntityView.h"
 #include "MassExecutionContext.h"
 #include "Queries/TypedElementExtendedQueryStore.h"
 #include "TypedElementDatabase.h"
 
-struct FMassContextForwarderShared
+template<typename T>
+struct FMassContextCommon : public T
 {
-	static const void* GetColumn(const FMassExecutionContext& Context, const UScriptStruct* ColumnType)
+	~FMassContextCommon() override = default;
+
+	const void* GetColumn(const UScriptStruct* ColumnType) const override
 	{
 		return Context.GetFragmentView(ColumnType).GetData();
 	}
 
-	static void* GetMutableColumn(FMassExecutionContext& Context, const UScriptStruct* ColumnType)
+	void* GetMutableColumn(const UScriptStruct* ColumnType) override
 	{
 		return Context.GetMutableFragmentView(ColumnType).GetData();
 	}
 
-	static void GetColumnsUnguarded(FMassExecutionContext& Context, int32 TypeCount, char** RetrievedAddresses,
-		const TWeakObjectPtr<const UScriptStruct>* ColumnTypes, const ITypedElementDataStorageInterface::EQueryAccessType* AccessTypes)
+	void GetColumns(TArrayView<char*> RetrievedAddresses, TConstArrayView<TWeakObjectPtr<const UScriptStruct>> ColumnTypes,
+		TConstArrayView<TypedElementDataStorage::EQueryAccessType> AccessTypes) override
+	{
+		checkf(RetrievedAddresses.Num() == ColumnTypes.Num(), TEXT("Unable to retrieve a batch of columns as the number of addresses "
+			"doesn't match the number of requested column."));
+		checkf(RetrievedAddresses.Num() == AccessTypes.Num(), TEXT("Unable to retrieve a batch of columns as the number of addresses "
+			"doesn't match the number of access types."));
+
+		GetColumnsUnguarded(ColumnTypes.Num(), RetrievedAddresses.GetData(), ColumnTypes.GetData(), AccessTypes.GetData());
+	}
+	
+	void GetColumnsUnguarded(int32 TypeCount, char** RetrievedAddresses, const TWeakObjectPtr<const UScriptStruct>* ColumnTypes,
+		const TypedElementDataStorage::EQueryAccessType* AccessTypes) override
 	{
 		for (int32 Index = 0; Index < TypeCount; ++Index)
 		{
@@ -36,24 +52,12 @@ struct FMassContextForwarderShared
 		}
 	}
 
-	static void GetColumns(FMassExecutionContext& Context, TArrayView<char*> RetrievedAddresses, 
-		TConstArrayView<TWeakObjectPtr<const UScriptStruct>> ColumnTypes,
-		TConstArrayView<ITypedElementDataStorageInterface::EQueryAccessType> AccessTypes)
-	{
-		checkf(RetrievedAddresses.Num() == ColumnTypes.Num(), TEXT("Unable to retrieve a batch of columns as the number of addresses "
-			"doesn't match the number of requested column."));
-		checkf(RetrievedAddresses.Num() == AccessTypes.Num(), TEXT("Unable to retrieve a batch of columns as the number of addresses "
-			"doesn't match the number of access types."));
-
-		GetColumnsUnguarded(Context, ColumnTypes.Num(), RetrievedAddresses.GetData(), ColumnTypes.GetData(), AccessTypes.GetData());
-	}
-
-	static uint32 GetRowCount(const FMassExecutionContext& Context)
+	uint32 GetRowCount() const override
 	{
 		return Context.GetNumEntities();
 	}
-
-	static TConstArrayView<TypedElementRowHandle> GetRowHandles(const FMassExecutionContext& Context)
+	
+	TConstArrayView<TypedElementRowHandle> GetRowHandles() const override
 	{
 		static_assert(
 			sizeof(TypedElementRowHandle) == sizeof(FMassEntityHandle) && alignof(TypedElementRowHandle) == alignof(FMassEntityHandle),
@@ -61,41 +65,98 @@ struct FMassContextForwarderShared
 		TConstArrayView<FMassEntityHandle> Entities = Context.GetEntities();
 		return TConstArrayView<TypedElementRowHandle>(reinterpret_cast<const TypedElementRowHandle*>(Entities.GetData()), Entities.Num());
 	}
+
+protected:
+	explicit FMassContextCommon(FMassExecutionContext& InContext)
+		: Context(InContext)
+	{}
+
+	FMassExecutionContext& Context;
 };
 
-struct FMassContextForwarder final : public ITypedElementDataStorageInterface::IQueryContext
+struct FMassDirectContextForwarder final : public FMassContextCommon<ITypedElementDataStorageInterface::IDirectQueryContext>
 {
-	FMassContextForwarder(ITypedElementDataStorageInterface::FQueryDescription& InQueryDescription, FMassExecutionContext& InContext, 
-		FTypedElementExtendedQueryStore& InQueryStore)
-		: QueryDescription(InQueryDescription)
-		, Context(InContext)
-		, QueryStore(InQueryStore)
+	explicit FMassDirectContextForwarder(FMassExecutionContext& InContext)
+		: FMassContextCommon(InContext)
 	{}
 
-	~FMassContextForwarder() override = default;
+	~FMassDirectContextForwarder() override = default;
+};
+
+struct FMassSubqueryContextForwarder final : public FMassContextCommon<ITypedElementDataStorageInterface::ISubqueryContext>
+{
+	explicit FMassSubqueryContextForwarder(FMassExecutionContext& InContext)
+		: FMassContextCommon(InContext)
+	{}
+
+	~FMassSubqueryContextForwarder() override = default;
+};
+
+struct FMassSingleRowSubqueryContextForwarder final : public ITypedElementDataStorageInterface::ISubqueryContext
+{
+	FMassSingleRowSubqueryContextForwarder(FMassEntityManager& InEntityManager, TypedElementDataStorage::RowHandle InRowHandle)
+		: EntityManager(InEntityManager)
+		, RowHandle(InRowHandle)
+	{}
+
+	~FMassSingleRowSubqueryContextForwarder() override = default;
 
 	const void* GetColumn(const UScriptStruct* ColumnType) const override
 	{
-		return FMassContextForwarderShared::GetColumn(Context, ColumnType);
+		return EntityManager.GetFragmentDataStruct(FMassEntityHandle::FromNumber(RowHandle), ColumnType).GetMemory();
 	}
 
 	void* GetMutableColumn(const UScriptStruct* ColumnType) override
 	{
-		return FMassContextForwarderShared::GetMutableColumn(Context, ColumnType);
+		return EntityManager.GetFragmentDataStruct(FMassEntityHandle::FromNumber(RowHandle), ColumnType).GetMemory();
 	}
 
 	void GetColumns(TArrayView<char*> RetrievedAddresses, TConstArrayView<TWeakObjectPtr<const UScriptStruct>> ColumnTypes,
-		TConstArrayView<ITypedElementDataStorageInterface::EQueryAccessType> AccessTypes) override
+		TConstArrayView<TypedElementDataStorage::EQueryAccessType> AccessTypes) override
 	{
-		FMassContextForwarderShared::GetColumns(Context, RetrievedAddresses, ColumnTypes, AccessTypes);
+		checkf(RetrievedAddresses.Num() == ColumnTypes.Num(), TEXT("Unable to retrieve a batch of columns as the number of addresses "
+			"doesn't match the number of requested column."));
+		checkf(RetrievedAddresses.Num() == AccessTypes.Num(), TEXT("Unable to retrieve a batch of columns as the number of addresses "
+			"doesn't match the number of access types."));
+
+		GetColumnsUnguarded(ColumnTypes.Num(), RetrievedAddresses.GetData(), ColumnTypes.GetData(), AccessTypes.GetData());
 	}
 
 	void GetColumnsUnguarded(int32 TypeCount, char** RetrievedAddresses, const TWeakObjectPtr<const UScriptStruct>* ColumnTypes,
-		const ITypedElementDataStorageInterface::EQueryAccessType* AccessTypes) override
+		const TypedElementDataStorage::EQueryAccessType*) override
+	{
+		for (int32 Index = 0; Index < TypeCount; ++Index)
+		{
+			checkf(ColumnTypes->IsValid(), TEXT("Attempting to retrieve a column that is not available."));
+			*RetrievedAddresses = reinterpret_cast<char*>(GetMutableColumn(ColumnTypes->Get()));
+
+			++RetrievedAddresses;
+			++ColumnTypes;
+		}
+	}
+
+	uint32 GetRowCount() const override { return 1; }
+	TConstArrayView<TypedElementDataStorage::RowHandle> GetRowHandles() const override
 	{
-		FMassContextForwarderShared::GetColumnsUnguarded(Context, TypeCount, RetrievedAddresses, ColumnTypes, AccessTypes);
+		return TConstArrayView<TypedElementDataStorage::RowHandle>(&RowHandle, 1);
 	}
 
+private:
+	FMassEntityManager& EntityManager;
+	TypedElementDataStorage::RowHandle RowHandle;
+};
+
+struct FMassContextForwarder final : public FMassContextCommon<ITypedElementDataStorageInterface::IQueryContext>
+{
+	FMassContextForwarder(ITypedElementDataStorageInterface::FQueryDescription& InQueryDescription, FMassExecutionContext& InContext, 
+		FTypedElementExtendedQueryStore& InQueryStore)
+		: FMassContextCommon(InContext)
+		, QueryDescription(InQueryDescription)
+		, QueryStore(InQueryStore)
+	{}
+
+	~FMassContextForwarder() override = default;
+
 	UObject* GetMutableDependency(const UClass* DependencyClass) override
 	{
 		return Context.GetMutableSubsystem<USubsystem>(const_cast<UClass*>(DependencyClass));
@@ -131,16 +192,6 @@ struct FMassContextForwarder final : public ITypedElementDataStorageInterface::I
 		}
 	}
 
-	uint32 GetRowCount() const override
-	{
-		return FMassContextForwarderShared::GetRowCount(Context);
-	}
-
-	TConstArrayView<TypedElementRowHandle> GetRowHandles() const
-	{
-		return FMassContextForwarderShared::GetRowHandles(Context);
-	}
-
 	void RemoveRow(TypedElementRowHandle Row) override
 	{
 		Context.Defer().DestroyEntity(FMassEntityHandle::FromNumber(Row));
@@ -226,7 +277,7 @@ struct FMassContextForwarder final : public ITypedElementDataStorageInterface::I
 		}
 	}
 
-	ITypedElementDataStorageInterface::FQueryResult RunQuery(TypedElementQueryHandle Query) override
+	TypedElementDataStorage::FQueryResult RunQuery(TypedElementQueryHandle Query) override
 	{
 		FTypedElementExtendedQueryStore::Handle Handle;
 		Handle.Handle = Query;
@@ -234,60 +285,47 @@ struct FMassContextForwarder final : public ITypedElementDataStorageInterface::I
 		return QueryStore.RunQuery(Context.GetEntityManagerChecked(), Handle);
 	}
 
-	ITypedElementDataStorageInterface::FQueryResult RunSubquery(int32 SubqueryIndex) override
+	TypedElementDataStorage::FQueryResult RunSubquery(int32 SubqueryIndex) override
 	{
 		return SubqueryIndex < QueryDescription.Subqueries.Num() ?
 			RunQuery(QueryDescription.Subqueries[SubqueryIndex]) :
-			ITypedElementDataStorageInterface::FQueryResult{};
-	}
-
-	ITypedElementDataStorageInterface::FQueryDescription& QueryDescription;
-	FMassExecutionContext& Context;
-	FTypedElementExtendedQueryStore& QueryStore;
-};
-
-struct FMassDirectContextForwarder final : public ITypedElementDataStorageInterface::IDirectQueryContext
-{
-	explicit FMassDirectContextForwarder(FMassExecutionContext& Context)
-		: Context(Context)
-	{}
-
-	~FMassDirectContextForwarder() override = default;
-
-	FMassExecutionContext& Context;
-
-	const void* GetColumn(const UScriptStruct* ColumnType) const override
-	{
-		return FMassContextForwarderShared::GetColumn(Context, ColumnType);
+			TypedElementDataStorage::FQueryResult{};
 	}
 
-	void* GetMutableColumn(const UScriptStruct* ColumnType) override
+	TypedElementDataStorage::FQueryResult RunSubquery(int32 SubqueryIndex, TypedElementDataStorage::SubqueryCallbackRef Callback) override
 	{
-		return FMassContextForwarderShared::GetMutableColumn(Context, ColumnType);
+		if (SubqueryIndex < QueryDescription.Subqueries.Num())
+		{
+			FTypedElementExtendedQueryStore::Handle Handle;
+			Handle.Handle = QueryDescription.Subqueries[SubqueryIndex];
+			return QueryStore.RunQuery(Context.GetEntityManagerChecked(), Handle, Callback);
+		}
+		else
+		{
+			return TypedElementDataStorage::FQueryResult{};
+		}
 	}
 
-	void GetColumns(TArrayView<char*> RetrievedAddresses, TConstArrayView<TWeakObjectPtr<const UScriptStruct>> ColumnTypes,
-		TConstArrayView<ITypedElementDataStorageInterface::EQueryAccessType> AccessTypes) override
+	TypedElementDataStorage::FQueryResult RunSubquery(int32 SubqueryIndex, TypedElementDataStorage::RowHandle Row,
+		TypedElementDataStorage::SubqueryCallbackRef Callback) override
 	{
-		FMassContextForwarderShared::GetColumns(Context, RetrievedAddresses, ColumnTypes, AccessTypes);
+		if (SubqueryIndex < QueryDescription.Subqueries.Num())
+		{
+			FTypedElementExtendedQueryStore::Handle Handle;
+			Handle.Handle = QueryDescription.Subqueries[SubqueryIndex];
+			return QueryStore.RunQuery(Context.GetEntityManagerChecked(), Handle, Row, Callback);
+		}
+		else
+		{
+			return TypedElementDataStorage::FQueryResult{};
+		}
 	}
 
-	void GetColumnsUnguarded(int32 TypeCount, char** RetrievedAddresses, const TWeakObjectPtr<const UScriptStruct>* ColumnTypes,
-		const ITypedElementDataStorageInterface::EQueryAccessType* AccessTypes) override
-	{
-		FMassContextForwarderShared::GetColumnsUnguarded(Context, TypeCount, RetrievedAddresses, ColumnTypes, AccessTypes);
-	}
+	ITypedElementDataStorageInterface::FQueryDescription& QueryDescription;
+	FTypedElementExtendedQueryStore& QueryStore;
+};
 
-	uint32 GetRowCount() const override
-	{
-		return FMassContextForwarderShared::GetRowCount(Context);
-	}
 
-	TConstArrayView<TypedElementRowHandle> GetRowHandles() const
-	{
-		return FMassContextForwarderShared::GetRowHandles(Context);
-	}
-};
 
 
 
@@ -408,28 +446,29 @@ bool FTypedElementQueryProcessorData::PrepareCachedDependenciesOnQuery(
 		checkf(Types->IsValid(), TEXT("Attempting to retrieve a dependency type that's no longer available."));
 		
 		if (EnumHasAnyFlags(*Flags, ITypedElementDataStorageInterface::EQueryDependencyFlags::AlwaysRefresh) || !Caches->IsValid())
-		*Caches = EnumHasAnyFlags(*Flags, ITypedElementDataStorageInterface::EQueryDependencyFlags::ReadOnly)
-			? const_cast<USubsystem*>(Context.GetSubsystem<USubsystem>(const_cast<UClass*>(Types->Get())))
-			: Context.GetMutableSubsystem<USubsystem>(const_cast<UClass*>(Types->Get()));
-
-		if (*Caches != nullptr)
 		{
-			++Types;
-			++Flags;
-			++Caches;
-		}
-		else
-		{
-			checkf(false, TEXT("Unable to retrieve instance of depencendy '%s'."), *((*Types)->GetName()));
-			return false;
+			*Caches = EnumHasAnyFlags(*Flags, ITypedElementDataStorageInterface::EQueryDependencyFlags::ReadOnly)
+				? const_cast<USubsystem*>(Context.GetSubsystem<USubsystem>(const_cast<UClass*>(Types->Get())))
+				: Context.GetMutableSubsystem<USubsystem>(const_cast<UClass*>(Types->Get()));
+			if (*Caches != nullptr)
+			{
+				++Types;
+				++Flags;
+				++Caches;
+			}
+			else
+			{
+				checkf(false, TEXT("Unable to retrieve instance of dependency '%s'."), *((*Types)->GetName()));
+				return false;
+			}
 		}
 	}
 	return true;
 }
 
-ITypedElementDataStorageInterface::FQueryResult FTypedElementQueryProcessorData::Execute(
-	ITypedElementDataStorageInterface::DirectQueryCallbackRef& Callback,
-	ITypedElementDataStorageInterface::FQueryDescription& Description, 
+TypedElementDataStorage::FQueryResult FTypedElementQueryProcessorData::Execute(
+	TypedElementDataStorage::DirectQueryCallbackRef& Callback,
+	TypedElementDataStorage::FQueryDescription& Description,
 	FMassEntityQuery& NativeQuery, 
 	FMassEntityManager& EntityManager)
 {
@@ -440,21 +479,61 @@ ITypedElementDataStorageInterface::FQueryResult FTypedElementQueryProcessorData:
 	NativeQuery.ForEachEntityChunk(EntityManager, Context,
 		[&Result, &Callback, &Description](FMassExecutionContext& Context)
 		{
-			if (PrepareCachedDependenciesOnQuery(Description, Context))
-			{
-				FMassDirectContextForwarder QueryContext(Context);
-				Callback(Description, QueryContext);
-				Result.Count += Context.GetNumEntities();
-			}
-			else
-			{
-				Result.Completed = ITypedElementDataStorageInterface::FQueryResult::ECompletion::MissingDependency;
-			}
+			// No need to cache any subsystem dependencies as these are not accessible from a direct query.
+			FMassDirectContextForwarder QueryContext(Context);
+			Callback(Description, QueryContext);
+			Result.Count += Context.GetNumEntities();
+		}
+	);
+	return Result;
+}
+
+TypedElementDataStorage::FQueryResult FTypedElementQueryProcessorData::Execute(
+	TypedElementDataStorage::SubqueryCallbackRef& Callback,
+	TypedElementDataStorage::FQueryDescription& Description,
+	FMassEntityQuery& NativeQuery,
+	FMassEntityManager& EntityManager)
+{
+	FMassExecutionContext Context(EntityManager);
+	ITypedElementDataStorageInterface::FQueryResult Result;
+	Result.Completed = ITypedElementDataStorageInterface::FQueryResult::ECompletion::Fully;
+
+	NativeQuery.ForEachEntityChunk(EntityManager, Context,
+		[&Result, &Callback, &Description](FMassExecutionContext& Context)
+		{
+			// No need to cache any subsystem dependencies as these are not accessible from a subquery.
+			FMassSubqueryContextForwarder QueryContext(Context);
+			Callback(Description, QueryContext);
+			Result.Count += Context.GetNumEntities();
 		}
 	);
 	return Result;
 }
 
+TypedElementDataStorage::FQueryResult FTypedElementQueryProcessorData::Execute(
+	TypedElementDataStorage::SubqueryCallbackRef& Callback,
+	TypedElementDataStorage::FQueryDescription& Description,
+	TypedElementDataStorage::RowHandle RowHandle,
+	FMassEntityQuery& NativeQuery,
+	FMassEntityManager& EntityManager)
+{
+	ITypedElementDataStorageInterface::FQueryResult Result;
+	Result.Completed = ITypedElementDataStorageInterface::FQueryResult::ECompletion::Fully;
+
+	FMassEntityHandle NativeEntity = FMassEntityHandle::FromNumber(RowHandle);
+	if (EntityManager.IsEntityActive(NativeEntity))
+	{
+		FMassArchetypeHandle NativeArchetype = EntityManager.GetArchetypeForEntityUnsafe(NativeEntity);
+		if (NativeQuery.DoesArchetypeMatchRequirements(NativeArchetype))
+		{
+			FMassSingleRowSubqueryContextForwarder QueryContext(EntityManager, RowHandle);
+			Callback(Description, QueryContext);
+			Result.Count = 1;
+		}
+	}
+	return Result;
+}
+
 void FTypedElementQueryProcessorData::Execute(FMassEntityManager& EntityManager, FMassExecutionContext& Context)
 {
 	checkf(ParentQuery, TEXT("A query callback was registered for execution without an associated query."));
diff --git a/Engine/Plugins/Experimental/TypedElementsDataStorage/Source/TypedElementsDataStorage/Private/Processors/TypedElementProcessorAdaptors.h b/Engine/Plugins/Experimental/TypedElementsDataStorage/Source/TypedElementsDataStorage/Private/Processors/TypedElementProcessorAdaptors.h
--- a/Engine/Plugins/Experimental/TypedElementsDataStorage/Source/TypedElementsDataStorage/Private/Processors/TypedElementProcessorAdaptors.h
+++ b/Engine/Plugins/Experimental/TypedElementsDataStorage/Source/TypedElementsDataStorage/Private/Processors/TypedElementProcessorAdaptors.h
@@ -41,11 +41,22 @@ struct FTypedElementQueryProcessorData
 	static EMassProcessingPhase MapToMassProcessingPhase(ITypedElementDataStorageInterface::EQueryTickPhase Phase);
 	FString GetProcessorName() const;
 
-	static ITypedElementDataStorageInterface::FQueryResult Execute(
-		ITypedElementDataStorageInterface::DirectQueryCallbackRef& Callback,
-		ITypedElementDataStorageInterface::FQueryDescription& Description, 
+	static TypedElementDataStorage::FQueryResult Execute(
+		TypedElementDataStorage::DirectQueryCallbackRef& Callback,
+		TypedElementDataStorage::FQueryDescription& Description,
 		FMassEntityQuery& NativeQuery, 
 		FMassEntityManager& EntityManager);
+	static TypedElementDataStorage::FQueryResult Execute(
+		TypedElementDataStorage::SubqueryCallbackRef& Callback,
+		TypedElementDataStorage::FQueryDescription& Description,
+		FMassEntityQuery& NativeQuery,
+		FMassEntityManager& EntityManager);
+	static TypedElementDataStorage::FQueryResult Execute(
+		TypedElementDataStorage::SubqueryCallbackRef& Callback,
+		TypedElementDataStorage::FQueryDescription& Description,
+		TypedElementDataStorage::RowHandle RowHandle,
+		FMassEntityQuery& NativeQuery,
+		FMassEntityManager& EntityManager);
 	void Execute(FMassEntityManager& EntityManager, FMassExecutionContext& Context);
 
 	static bool PrepareCachedDependenciesOnQuery(
diff --git a/Engine/Plugins/Experimental/TypedElementsDataStorage/Source/TypedElementsDataStorage/Private/Queries/TypedElementExtendedQueryStore.cpp b/Engine/Plugins/Experimental/TypedElementsDataStorage/Source/TypedElementsDataStorage/Private/Queries/TypedElementExtendedQueryStore.cpp
--- a/Engine/Plugins/Experimental/TypedElementsDataStorage/Source/TypedElementsDataStorage/Private/Queries/TypedElementExtendedQueryStore.cpp
+++ b/Engine/Plugins/Experimental/TypedElementsDataStorage/Source/TypedElementsDataStorage/Private/Queries/TypedElementExtendedQueryStore.cpp
@@ -140,7 +140,7 @@ void FTypedElementExtendedQueryStore::ListAliveEntries(const ListAliveEntriesCal
 	Queries.ListAliveEntries(Callback);
 }
 
-ITypedElementDataStorageInterface::FQueryResult FTypedElementExtendedQueryStore::RunQuery(FMassEntityManager& EntityManager, Handle Query)
+TypedElementDataStorage::FQueryResult FTypedElementExtendedQueryStore::RunQuery(FMassEntityManager& EntityManager, Handle Query)
 {
 	using ActionType = ITypedElementDataStorageInterface::FQueryDescription::EActionType;
 	using CompletionType = ITypedElementDataStorageInterface::FQueryResult::ECompletion;
@@ -180,8 +180,9 @@ ITypedElementDataStorageInterface::FQueryResult FTypedElementExtendedQueryStore:
 	return Result;
 }
 
-ITypedElementDataStorageInterface::FQueryResult FTypedElementExtendedQueryStore::RunQuery(FMassEntityManager& EntityManager, Handle Query,
-	ITypedElementDataStorageInterface::DirectQueryCallbackRef Callback)
+template<typename CallbackReference>
+TypedElementDataStorage::FQueryResult FTypedElementExtendedQueryStore::RunQueryCallbackCommon(FMassEntityManager& EntityManager, Handle Query,
+	CallbackReference Callback)
 {
 	using ActionType = ITypedElementDataStorageInterface::FQueryDescription::EActionType;
 	using CompletionType = ITypedElementDataStorageInterface::FQueryResult::ECompletion;
@@ -224,6 +225,62 @@ ITypedElementDataStorageInterface::FQueryResult FTypedElementExtendedQueryStore:
 	return Result;
 }
 
+TypedElementDataStorage::FQueryResult FTypedElementExtendedQueryStore::RunQuery(FMassEntityManager& EntityManager, Handle Query,
+	TypedElementDataStorage::DirectQueryCallbackRef Callback)
+{
+	return RunQueryCallbackCommon(EntityManager, Query, Callback);
+}
+
+TypedElementDataStorage::FQueryResult FTypedElementExtendedQueryStore::RunQuery(FMassEntityManager& EntityManager, Handle Query,
+	TypedElementDataStorage::SubqueryCallbackRef Callback)
+{
+	return RunQueryCallbackCommon(EntityManager, Query, Callback);
+}
+
+TypedElementDataStorage::FQueryResult FTypedElementExtendedQueryStore::RunQuery(FMassEntityManager& EntityManager, Handle Query,
+	TypedElementRowHandle Row, TypedElementDataStorage::SubqueryCallbackRef Callback)
+{
+	using ActionType = ITypedElementDataStorageInterface::FQueryDescription::EActionType;
+	using CompletionType = ITypedElementDataStorageInterface::FQueryResult::ECompletion;
+
+	ITypedElementDataStorageInterface::FQueryResult Result;
+
+	if (FTypedElementExtendedQuery* QueryData = Get(Query))
+	{
+		switch (QueryData->Description.Action)
+		{
+		case ActionType::None:
+			Result.Completed = CompletionType::Fully;
+			break;
+		case ActionType::Select:
+			if (!QueryData->Processor.IsValid())
+			{
+				Result = FTypedElementQueryProcessorData::Execute(
+					Callback, QueryData->Description, Row, QueryData->NativeQuery, EntityManager);
+			}
+			else
+			{
+				Result.Completed = CompletionType::Unsupported;
+			}
+			break;
+		case ActionType::Count:
+			// Only the count is requested so no need to trigger the callback.
+			Result.Count = 1;
+			Result.Completed = CompletionType::Fully;
+			break;
+		default:
+			Result.Completed = CompletionType::Unsupported;
+			break;
+		}
+	}
+	else
+	{
+		Result.Completed = CompletionType::Unavailable;
+	}
+
+	return Result;
+}
+
 void FTypedElementExtendedQueryStore::RunPhasePreambleQueries(FMassEntityManager& EntityManager,
 	ITypedElementDataStorageInterface::EQueryTickPhase Phase, float DeltaTime)
 {
diff --git a/Engine/Plugins/Experimental/TypedElementsDataStorage/Source/TypedElementsDataStorage/Private/Queries/TypedElementExtendedQueryStore.h b/Engine/Plugins/Experimental/TypedElementsDataStorage/Source/TypedElementsDataStorage/Private/Queries/TypedElementExtendedQueryStore.h
--- a/Engine/Plugins/Experimental/TypedElementsDataStorage/Source/TypedElementsDataStorage/Private/Queries/TypedElementExtendedQueryStore.h
+++ b/Engine/Plugins/Experimental/TypedElementsDataStorage/Source/TypedElementsDataStorage/Private/Queries/TypedElementExtendedQueryStore.h
@@ -80,9 +80,13 @@ public:
 	 * @section Execution
 	 * @description Various functions to run queries.
 	 */
-	ITypedElementDataStorageInterface::FQueryResult RunQuery(FMassEntityManager& EntityManager, Handle Query);
-	ITypedElementDataStorageInterface::FQueryResult RunQuery(FMassEntityManager& EntityManager, Handle Query,
-		ITypedElementDataStorageInterface::DirectQueryCallbackRef Callback);
+	TypedElementDataStorage::FQueryResult RunQuery(FMassEntityManager& EntityManager, Handle Query);
+	TypedElementDataStorage::FQueryResult RunQuery(FMassEntityManager& EntityManager, Handle Query,
+		TypedElementDataStorage::DirectQueryCallbackRef Callback);
+	TypedElementDataStorage::FQueryResult RunQuery(FMassEntityManager& EntityManager, Handle Query,
+		TypedElementDataStorage::SubqueryCallbackRef Callback);
+	TypedElementDataStorage::FQueryResult RunQuery(FMassEntityManager& EntityManager, Handle Query,
+		TypedElementRowHandle Row, TypedElementDataStorage::SubqueryCallbackRef Callback);
 	void RunPhasePreambleQueries(FMassEntityManager& EntityManager, ITypedElementDataStorageInterface::EQueryTickPhase Phase, float DeltaTime);
 	void RunPhasePostambleQueries(FMassEntityManager& EntityManager, ITypedElementDataStorageInterface::EQueryTickPhase Phase, float DeltaTime);
 
@@ -109,6 +113,9 @@ private:
 		bool bRequiresMainThread{ false };
 	};
 
+	template<typename CallbackReference>
+	TypedElementDataStorage::FQueryResult RunQueryCallbackCommon(FMassEntityManager& EntityManager, Handle Query, CallbackReference Callback);
+
 	FMassEntityQuery& SetupNativeQuery(ITypedElementDataStorageInterface::FQueryDescription& Query, FTypedElementExtendedQuery& StoredQuery);
 	bool SetupSelectedColumns(ITypedElementDataStorageInterface::FQueryDescription& Query, FMassEntityQuery& NativeQuery);
 	bool SetupConditions(ITypedElementDataStorageInterface::FQueryDescription& Query, FMassEntityQuery& NativeQuery);
diff --git a/Engine/Source/Runtime/TypedElementFramework/Public/Elements/Common/TypedElementQueryDescription.h b/Engine/Source/Runtime/TypedElementFramework/Public/Elements/Common/TypedElementQueryDescription.h
--- a/Engine/Source/Runtime/TypedElementFramework/Public/Elements/Common/TypedElementQueryDescription.h
+++ b/Engine/Source/Runtime/TypedElementFramework/Public/Elements/Common/TypedElementQueryDescription.h
@@ -22,8 +22,6 @@ namespace TypedElementDataStorage
 	using QueryCallbackRef = TFunctionRef<void(const FQueryDescription&, IQueryContext&)>;
 	using DirectQueryCallback = TFunction<void(const FQueryDescription&, IDirectQueryContext&)>;
 	using DirectQueryCallbackRef = TFunctionRef<void(const FQueryDescription&, IDirectQueryContext&)>;
-	using SubqueryCallback = TFunction<void(const FQueryDescription&, ISubqueryContext&)>;
-	using SubqueryCallbackRef = TFunctionRef<void(const FQueryDescription&, ISubqueryContext&)>;
 	
 	struct FQueryDescription final
 	{
diff --git a/Engine/Source/Runtime/TypedElementFramework/Public/Elements/Framework/TypedElementQueryBuilder.h b/Engine/Source/Runtime/TypedElementFramework/Public/Elements/Framework/TypedElementQueryBuilder.h
--- a/Engine/Source/Runtime/TypedElementFramework/Public/Elements/Framework/TypedElementQueryBuilder.h
+++ b/Engine/Source/Runtime/TypedElementFramework/Public/Elements/Framework/TypedElementQueryBuilder.h
@@ -253,8 +253,11 @@ namespace TypedElementQueryBuilder
 		inline void RemoveColumns(TypedElementRowHandle Row, TConstArrayView<const UScriptStruct*> ColumnTypes) override;
 		inline void RemoveColumns(TConstArrayView<TypedElementRowHandle> Rows, TConstArrayView<const UScriptStruct*> ColumnTypes) override;
 
-		inline ITypedElementDataStorageInterface::FQueryResult RunQuery(TypedElementQueryHandle Query) override;
-		inline ITypedElementDataStorageInterface::FQueryResult RunSubquery(int32 SubqueryIndex) override;
+		inline TypedElementDataStorage::FQueryResult RunQuery(TypedElementQueryHandle Query) override;
+		inline TypedElementDataStorage::FQueryResult RunSubquery(int32 SubqueryIndex) override;
+		inline TypedElementDataStorage::FQueryResult RunSubquery(int32 SubqueryIndex, TypedElementDataStorage::SubqueryCallbackRef Callback) override;
+		inline TypedElementDataStorage::FQueryResult RunSubquery(int32 SubqueryIndex, TypedElementDataStorage::RowHandle Row,
+			TypedElementDataStorage::SubqueryCallbackRef Callback) override;
 
 		ITypedElementDataStorageInterface::IQueryContext& ParentContext;
 		const ITypedElementDataStorageInterface::FQueryDescription& Description;
@@ -317,9 +320,9 @@ namespace TypedElementQueryBuilder
 	};
 
 	template<typename Function>
-	ITypedElementDataStorageInterface::DirectQueryCallback CreateDirectQueryCallbackBinding(Function&& Callback);
+	TypedElementDataStorage::DirectQueryCallback CreateDirectQueryCallbackBinding(Function&& Callback);
 	template<typename Function>
-	ITypedElementDataStorageInterface::DirectQueryCallback CreateSubqueryCallbackBinding(Function&& Callback);
+	TypedElementDataStorage::SubqueryCallback CreateSubqueryCallbackBinding(Function&& Callback);
 
 } // namespace TypedElementQueryBuilder
 
diff --git a/Engine/Source/Runtime/TypedElementFramework/Public/Elements/Framework/TypedElementQueryBuilder.inl b/Engine/Source/Runtime/TypedElementFramework/Public/Elements/Framework/TypedElementQueryBuilder.inl
--- a/Engine/Source/Runtime/TypedElementFramework/Public/Elements/Framework/TypedElementQueryBuilder.inl
+++ b/Engine/Source/Runtime/TypedElementFramework/Public/Elements/Framework/TypedElementQueryBuilder.inl
@@ -219,6 +219,18 @@ namespace TypedElementQueryBuilder
 		return ParentContext.RunSubquery(SubqueryIndex);
 	}
 
+	TypedElementDataStorage::FQueryResult FQueryContextForwarder::RunSubquery(int32 SubqueryIndex,
+		TypedElementDataStorage::SubqueryCallbackRef Callback)
+	{
+		return ParentContext.RunSubquery(SubqueryIndex, Callback);
+	}
+
+	TypedElementDataStorage::FQueryResult FQueryContextForwarder::RunSubquery(int32 SubqueryIndex,
+		TypedElementDataStorage::RowHandle Row, TypedElementDataStorage::SubqueryCallbackRef Callback)
+	{
+		return ParentContext.RunSubquery(SubqueryIndex, Row, Callback);
+	}
+
 	//
 	// FCachedQueryContext
 	//
@@ -1179,7 +1191,7 @@ e.g. void(IDirectQueryContext& Context, TypedElementRowHandle Row, ColumnType0&
 	}
 
 	template<typename Function>
-	TypedElementDataStorage::DirectQueryCallback CreateSubqueryCallbackBinding(Function&& Callback)
+	TypedElementDataStorage::SubqueryCallback CreateSubqueryCallbackBinding(Function&& Callback)
 	{
 		static_assert(Internal::IsValidSelectFunctionSignature<ITypedElementDataStorageInterface::ISubqueryContext, Function>(),
 			R"(The function provided to the Query Builder's CreateSubqueryCallbackBinding call wasn't invocable or doesn't contain a supported combination of arguments.
diff --git a/Engine/Source/Runtime/TypedElementFramework/Public/Elements/Interfaces/TypedElementQueryStorageInterfaces.h b/Engine/Source/Runtime/TypedElementFramework/Public/Elements/Interfaces/TypedElementQueryStorageInterfaces.h
--- a/Engine/Source/Runtime/TypedElementFramework/Public/Elements/Interfaces/TypedElementQueryStorageInterfaces.h
+++ b/Engine/Source/Runtime/TypedElementFramework/Public/Elements/Interfaces/TypedElementQueryStorageInterfaces.h
@@ -13,6 +13,12 @@ class UScriptStruct;
 
 namespace TypedElementDataStorage
 {
+	struct FQueryDescription; 
+	struct ISubqueryContext;
+
+	using SubqueryCallback = TFunction<void(const FQueryDescription&, ISubqueryContext&)>;
+	using SubqueryCallbackRef = TFunctionRef<void(const FQueryDescription&, ISubqueryContext&)>;
+
 	/**
 	 * Base interface for any contexts provided to query callbacks.
 	 */
@@ -124,10 +130,26 @@ namespace TypedElementDataStorage
 		/**
 		 * Runs a previously created query. This version takes an arbitrary query, but is limited to running queries that do not directly
 		 * access data from rows such as count queries.
+		 * The returned result is a snap shot and values may change between phases.
 		 */
 		virtual FQueryResult RunQuery(QueryHandle Query) = 0;
-		/** Runs a subquery registered with the current query. The subquery index is in the order of registration with the query. */
+		/** 
+		 * Runs a subquery registered with the current query. The subquery index is in the order of registration with the query. Subqueries
+		 * are executed as part of their parent query and are not scheduled separately.
+		 */
 		virtual FQueryResult RunSubquery(int32 SubqueryIndex) = 0;
+		/** 
+		 * Runs the provided callback on a subquery registered with the current query. The subquery index is in the order of registration 
+		 * with the query. Subqueries are executed as part of their parent query and are not scheduled separately.
+		 */
+		virtual FQueryResult RunSubquery(int32 SubqueryIndex, SubqueryCallbackRef Callback) = 0;
+		/** 
+		 * Runs the provided callback on a subquery registered with the current query for the exact provided row. The subquery index is in 
+		 * the order of registration with the query. If the row handle is in a table that doesn't match the selected subquery the callback
+		 * will not be called. Check the count in the returned results to determine if the callback was called or not. Subqueries
+		 * are executed as part of their parent query and are not scheduled separately.
+		 */
+		virtual FQueryResult RunSubquery(int32 SubqueryIndex, RowHandle Row, SubqueryCallbackRef Callback) = 0;
 
 
 
