Moved several query callbacks from TEDS over to using the safer subqueries.

diff --git a/Engine/Plugins/Experimental/TypedElementsDataStorage/Source/TypedElementsDataStorage/Private/Compatibility/TypedElementActorLabelQueries.cpp b/Engine/Plugins/Experimental/TypedElementsDataStorage/Source/TypedElementsDataStorage/Private/Compatibility/TypedElementActorLabelQueries.cpp
--- a/Engine/Plugins/Experimental/TypedElementsDataStorage/Source/TypedElementsDataStorage/Private/Compatibility/TypedElementActorLabelQueries.cpp
+++ b/Engine/Plugins/Experimental/TypedElementsDataStorage/Source/TypedElementsDataStorage/Private/Compatibility/TypedElementActorLabelQueries.cpp
@@ -38,9 +38,9 @@ FAutoConsoleCommandWithOutputDevice PrintActorLabelsConsoleCommand(
 				{
 					FString Message;
 					Output.Log(TEXT("The Typed Elements Data Storage has the following actors:"));
-					DataStorage->RunQuery(LabelQuery, [&Output, &Message](const DSI::FQueryDescription&, DSI::IDirectQueryContext& Context)
+					DataStorage->RunQuery(LabelQuery, CreateDirectQueryCallbackBinding(
+						[&Output, &Message](DSI::IDirectQueryContext& Context, const FTypedElementLabelColumn* Labels)
 						{
-							const FTypedElementLabelColumn* Labels = Context.GetColumn<FTypedElementLabelColumn>();
 							const uint32 Count = Context.GetRowCount();
 
 							const FTypedElementLabelColumn* LabelsIt = Labels;
@@ -63,7 +63,7 @@ FAutoConsoleCommandWithOutputDevice PrintActorLabelsConsoleCommand(
 							}
 
 							Output.Log(Message);
-						});
+						}));
 					Output.Log(TEXT("End of Typed Elements Data Storage actors list."));
 				}
 			}
diff --git a/Engine/Plugins/Experimental/TypedElementsDataStorage/Source/TypedElementsDataStorageUI/Private/TypedElementExportedTextWidget.cpp b/Engine/Plugins/Experimental/TypedElementsDataStorage/Source/TypedElementsDataStorageUI/Private/TypedElementExportedTextWidget.cpp
--- a/Engine/Plugins/Experimental/TypedElementsDataStorage/Source/TypedElementsDataStorageUI/Private/TypedElementExportedTextWidget.cpp
+++ b/Engine/Plugins/Experimental/TypedElementsDataStorage/Source/TypedElementsDataStorageUI/Private/TypedElementExportedTextWidget.cpp
@@ -15,8 +15,8 @@
 // UTypedElementExportedTextWidgetFactory
 //
 
-static void UpdateExportedTextWidget(ITypedElementDataStorageInterface& DataStorage, FTypedElementSlateWidgetReferenceColumn& Widget,
-	const FTypedElementScriptStructTypeInfoColumn& TypeInfo, const FTypedElementRowReferenceColumn& ReferencedRow)
+static void UpdateExportedTextWidget(const void* Data, FTypedElementSlateWidgetReferenceColumn& Widget, 
+	const FTypedElementScriptStructTypeInfoColumn& TypeInfo)
 {
 	TSharedPtr<SWidget> WidgetPointer = Widget.Widget.Pin();
 	checkf(WidgetPointer, TEXT("Referenced widget is not valid. A constructed widget may not have been cleaned up. This can "
@@ -27,45 +27,61 @@ static void UpdateExportedTextWidget(ITypedElementDataStorageInterface& DataStor
 		*(STextBlock::StaticWidgetClass().GetWidgetType().ToString()),
 		*(WidgetPointer->GetTypeAsString()));
 
+	FString Label;
+	TypeInfo.TypeInfo->ExportText(Label, Data, nullptr, nullptr, PPF_None, nullptr);
+	STextBlock* TextWidget = static_cast<STextBlock*>(WidgetPointer.Get());
+	FText Text = FText::FromString(MoveTemp(Label));
+	TextWidget->SetToolTipText(Text);
+	TextWidget->SetText(MoveTemp(Text));
+}
+
+static void UpdateExportedTextWidget(ITypedElementDataStorageInterface& DataStorage, FTypedElementSlateWidgetReferenceColumn& Widget,
+	const FTypedElementScriptStructTypeInfoColumn& TypeInfo, const FTypedElementRowReferenceColumn& ReferencedRow)
+{
 	if (void* Data = DataStorage.GetColumnData(ReferencedRow.Row, TypeInfo.TypeInfo.Get()))
 	{
-		FString Label;
-		TypeInfo.TypeInfo->ExportText(Label, Data, nullptr, nullptr, PPF_None, nullptr);
-		STextBlock* TextWidget = static_cast<STextBlock*>(WidgetPointer.Get());
-		FText Text = FText::FromString(MoveTemp(Label));
-		TextWidget->SetToolTipText(Text);
-		TextWidget->SetText(MoveTemp(Text));
+		UpdateExportedTextWidget(Data, Widget, TypeInfo);
 	}
 }
 
-void UTypedElementExportedTextWidgetFactory::RegisterQueries(ITypedElementDataStorageInterface& DataStorage) const
+static TypedElementQueryHandle RegisterUpdateCallback(ITypedElementDataStorageInterface& DataStorage, const UScriptStruct* Target)
 {
 	using namespace TypedElementQueryBuilder;
 	using DSI = ITypedElementDataStorageInterface;
-
-	DataStorage.RegisterQuery(
-		Select(TEXT("Sync exported text widgets"),
-		FProcessor(DSI::EQueryTickPhase::FrameEnd, DataStorage.GetQueryTickGroupName(DSI::EQueryTickGroups::SyncWidgets))
-			.ForceToGameThread(true),
-			[](FCachedQueryContext<UTypedElementDataStorageSubsystem>& Context, FTypedElementSlateWidgetReferenceColumn& Widget,
-				const FTypedElementScriptStructTypeInfoColumn& TypeInfo, const FTypedElementRowReferenceColumn& ReferencedRow)
+	namespace DS = TypedElementDataStorage;
+	
+	TypedElementQueryHandle TypeDataQuery = DataStorage.RegisterQuery(
+		Select()
+			.ReadOnly(Target)
+		.Where()
+			.Any<FTypedElementSyncFromWorldTag, FTypedElementSyncBackToWorldTag>()
+		.Compile());
+
+	FString Name = TEXT("Sync exported text widgets (");
+	Target->AppendName(Name);
+	Name += ')';
+
+	return DataStorage.RegisterQuery(
+		Select(FName(Name),
+			FProcessor(DSI::EQueryTickPhase::FrameEnd, DataStorage.GetQueryTickGroupName(DSI::EQueryTickGroups::SyncWidgets))
+				.ForceToGameThread(true),
+			[](
+				DS::IQueryContext& Context, 
+				FTypedElementSlateWidgetReferenceColumn& Widget,
+				const FTypedElementScriptStructTypeInfoColumn& TypeInfo,
+				const FTypedElementRowReferenceColumn& ReferencedRow)
 			{
-				UTypedElementDataStorageSubsystem& Subsystem = Context.GetCachedMutableDependency<UTypedElementDataStorageSubsystem>();
-				DSI* DataStorage = Subsystem.Get();
-				checkf(DataStorage, TEXT("FTypedElementsDataStorageUiModule tried to process widgets before the "
-					"Typed Elements Data Storage interface is available."));
-				
-				if (DataStorage->HasColumns<FTypedElementSyncFromWorldTag>(ReferencedRow.Row) ||
-					DataStorage->HasColumns<FTypedElementSyncBackToWorldTag>(ReferencedRow.Row))
-				{
-					UpdateExportedTextWidget(*DataStorage, Widget, TypeInfo, ReferencedRow);
-				}
-			}
-		)
-	.Where()
-		.All<FTypedElementExportedTextWidgetTag>()
-	.Compile());
-
+				Context.RunSubquery(0, ReferencedRow.Row, 
+					[&Widget, &TypeInfo](const DS::FQueryDescription&, DS::ISubqueryContext& SubqueryContext)
+					{
+						UpdateExportedTextWidget(SubqueryContext.GetColumn(TypeInfo.TypeInfo.Get()), Widget, TypeInfo);
+					});
+			})
+		.Where()
+			.All<FTypedElementExportedTextWidgetTag>()
+		.DependsOn()
+			.SubQuery(TypeDataQuery)
+		.Compile());
 }
 
 void UTypedElementExportedTextWidgetFactory::RegisterWidgetConstructors(ITypedElementDataStorageInterface& DataStorage,
@@ -110,11 +126,22 @@ bool FTypedElementExportedTextWidgetConstructor::FinalizeWidget(
 	TypedElementRowHandle Row,
 	const TSharedPtr<SWidget>& Widget)
 {
+	FTypedElementScriptStructTypeInfoColumn& TypeInfoColumn = *DataStorage->GetColumn<FTypedElementScriptStructTypeInfoColumn>(Row);
+
 	UpdateExportedTextWidget(
 		*DataStorage,
 		*DataStorage->GetColumn<FTypedElementSlateWidgetReferenceColumn>(Row),
-		*DataStorage->GetColumn<FTypedElementScriptStructTypeInfoColumn>(Row),
+		TypeInfoColumn,
 		*DataStorage->GetColumn<FTypedElementRowReferenceColumn>(Row));
+	
+	UTypedElementExportedTextWidgetFactory* Factory = 
+		UTypedElementExportedTextWidgetFactory::StaticClass()->GetDefaultObject<UTypedElementExportedTextWidgetFactory>();
+	if (Factory && !Factory->RegisteredTypes.Contains(TypeInfoColumn.TypeInfo))
+	{
+		RegisterUpdateCallback(*DataStorage, TypeInfoColumn.TypeInfo.Get());
+		Factory->RegisteredTypes.Add(TypeInfoColumn.TypeInfo);
+	}
+	
 	return true;
 }
 
diff --git a/Engine/Plugins/Experimental/TypedElementsDataStorage/Source/TypedElementsDataStorageUI/Private/TypedElementLabelWidget.cpp b/Engine/Plugins/Experimental/TypedElementsDataStorage/Source/TypedElementsDataStorageUI/Private/TypedElementLabelWidget.cpp
--- a/Engine/Plugins/Experimental/TypedElementsDataStorage/Source/TypedElementsDataStorageUI/Private/TypedElementLabelWidget.cpp
+++ b/Engine/Plugins/Experimental/TypedElementsDataStorage/Source/TypedElementsDataStorageUI/Private/TypedElementLabelWidget.cpp
@@ -66,28 +66,52 @@ void UTypedElementLabelWidgetFactory::RegisterQueries(ITypedElementDataStorageIn
 {
 	using namespace TypedElementQueryBuilder;
 	using DSI = ITypedElementDataStorageInterface;
+	
+
+	TypedElementQueryHandle UpdateLabelWidget = DataStorage.RegisterQuery(
+		Select()
+			.ReadOnly<FTypedElementLabelColumn>()
+		.Where()
+			.Any<FTypedElementSyncFromWorldTag, FTypedElementSyncBackToWorldTag>()
+			.None<FTypedElementLabelHashColumn>()
+		.Compile());
+
+	TypedElementQueryHandle UpdateLabelAndHashWidget = DataStorage.RegisterQuery(
+		Select()
+			.ReadOnly<FTypedElementLabelColumn, FTypedElementLabelHashColumn>()
+		.Where()
+			.Any<FTypedElementSyncFromWorldTag, FTypedElementSyncBackToWorldTag>()
+		.Compile());
 
 	DataStorage.RegisterQuery(
 		Select(
 			TEXT("Sync label to widget"),
 			FProcessor(DSI::EQueryTickPhase::FrameEnd, DataStorage.GetQueryTickGroupName(DSI::EQueryTickGroups::SyncWidgets))
 				.ForceToGameThread(true),
-			[](	FCachedQueryContext<UTypedElementDataStorageSubsystem>& Context, 
+			[](
+				DSI::IQueryContext& Context, 
 				FTypedElementSlateWidgetReferenceColumn& Widget,
 				FTypedElementU64IntValueCacheColumn& TextHash,
 				const FTypedElementLabelWidgetColumn& Config,
 				const FTypedElementRowReferenceColumn& Target)
 			{
-				DSI* DataStorage = Context.GetCachedMutableDependency<UTypedElementDataStorageSubsystem>().Get();
-				checkf(DataStorage, TEXT("FTypedElementsDataStorageUiModule tried to process widgets before the "
-					"Typed Elements Data Storage interface is available."));
-				
-				if (DataStorage->HasColumns<FTypedElementSyncFromWorldTag>(Target.Row) ||
-					DataStorage->HasColumns<FTypedElementSyncBackToWorldTag>(Target.Row))
-				{
-					SyncColumnsToWidget(DataStorage, Target.Row, TextHash, Widget.Widget, Config.bShowHashInTooltip);
-				}
+				Context.RunSubquery(0, Target.Row, CreateSubqueryCallbackBinding(
+					[&Widget](const FTypedElementLabelColumn& Label)
+					{
+						UpdateTextWidget(Widget.Widget, Label, nullptr);
+					}));
+				Context.RunSubquery(1, Target.Row, CreateSubqueryCallbackBinding(
+					[&Widget, &TextHash, &Config](const FTypedElementLabelColumn& Label, const FTypedElementLabelHashColumn& Hash)
+					{
+						if (Hash.LabelHash != TextHash.Value)
+						{
+							UpdateTextWidget(Widget.Widget, Label, Config.bShowHashInTooltip ? &Hash.LabelHash : nullptr);
+							TextHash.Value = Hash.LabelHash;
+						}
+					}));
 			})
+		.DependsOn()
+			.SubQuery({ UpdateLabelWidget, UpdateLabelAndHashWidget })
 		.Compile()
 	);
 }
diff --git a/Engine/Plugins/Experimental/TypedElementsDataStorage/Source/TypedElementsDataStorageUI/Private/TypedElementTransformHeadsUpWidget.cpp b/Engine/Plugins/Experimental/TypedElementsDataStorage/Source/TypedElementsDataStorageUI/Private/TypedElementTransformHeadsUpWidget.cpp
--- a/Engine/Plugins/Experimental/TypedElementsDataStorage/Source/TypedElementsDataStorageUI/Private/TypedElementTransformHeadsUpWidget.cpp
+++ b/Engine/Plugins/Experimental/TypedElementsDataStorage/Source/TypedElementsDataStorageUI/Private/TypedElementTransformHeadsUpWidget.cpp
@@ -133,7 +133,7 @@ private:
 SLATE_IMPLEMENT_WIDGET(STransformQuickDisplay)
 void STransformQuickDisplay::PrivateRegisterAttributes(FSlateAttributeInitializer&) {}
 
-static void UpdateTransformHeadsUpDisplay(ITypedElementDataStorageInterface& DataStorage, FTypedElementSlateWidgetReferenceColumn& Widget, InternalTransformHelpers::EAbnormalTransformTypes AbnormalTransformFlags)
+static void UpdateTransformHeadsUpDisplay(FTypedElementSlateWidgetReferenceColumn& Widget, InternalTransformHelpers::EAbnormalTransformTypes AbnormalTransformFlags)
 {
 	TSharedPtr<SWidget> WidgetPointer = Widget.Widget.Pin();
 	checkf(WidgetPointer, TEXT("Referenced widget is not valid. A constructed widget may not have been cleaned up. This can "
@@ -156,33 +156,34 @@ void UTypedElementTransformHeadsUpWidgetFactory::RegisterQueries(ITypedElementDa
 {
 	using namespace TypedElementQueryBuilder;
 	using DSI = ITypedElementDataStorageInterface;
+	namespace DS = TypedElementDataStorage;
 		
+	TypedElementQueryHandle UpdateTransformWidget = DataStorage.RegisterQuery(
+		Select()
+			.ReadOnly<FTypedElementLocalTransformColumn>()
+		.Where()
+			.Any<FTypedElementSyncFromWorldTag, FTypedElementSyncBackToWorldTag>()
+		.Compile());
 
 	DataStorage.RegisterQuery(
 		Select(TEXT("Sync Transform column to heads up display"),
 		FProcessor(DSI::EQueryTickPhase::FrameEnd, DataStorage.GetQueryTickGroupName(DSI::EQueryTickGroups::SyncWidgets))
 			.ForceToGameThread(true),
-			[](FCachedQueryContext<UTypedElementDataStorageSubsystem>& Context,
+			[](DS::IQueryContext& Context,
 				FTypedElementSlateWidgetReferenceColumn& Widget,
 				const FTypedElementRowReferenceColumn& ReferenceColumn)
 			{
-				UTypedElementDataStorageSubsystem& Subsystem = Context.GetCachedMutableDependency<UTypedElementDataStorageSubsystem>();
-				DSI* DataStorage = Subsystem.Get();
-				checkf(DataStorage, TEXT("FTypedElementsDataStorageUiModule tried to process widgets before the "
-					"Typed Elements Data Storage interface is available."));
-
-				if (DataStorage->HasColumns<FTypedElementSyncFromWorldTag>(ReferenceColumn.Row)
-					|| DataStorage->HasColumns<FTypedElementSyncBackToWorldTag>(ReferenceColumn.Row))
-				{
-					if (const FTypedElementLocalTransformColumn* TransformColumn = DataStorage->GetColumn<FTypedElementLocalTransformColumn>(ReferenceColumn.Row))
+				Context.RunSubquery(0, ReferenceColumn.Row, CreateSubqueryCallbackBinding(
+					[&Widget](const FTypedElementLocalTransformColumn& Transform)
 					{
-						UpdateTransformHeadsUpDisplay(*DataStorage, Widget, InternalTransformHelpers::GetAbnormalTransformTypes(TransformColumn->Transform));
-					}
-				}
+						UpdateTransformHeadsUpDisplay(Widget, InternalTransformHelpers::GetAbnormalTransformTypes(Transform.Transform));
+					}));
 			}
 		)
 	.Where()
 		.All<FTypedElementTransformHeadsUpWidgetTag>()
+	.DependsOn()
+		.SubQuery(UpdateTransformWidget)
 	.Compile());
 
 }
@@ -231,7 +232,8 @@ bool FTypedElementTransformHeadsUpWidgetConstructor::FinalizeWidget(
 	FTypedElementRowReferenceColumn& RefColumn = *DataStorage->GetColumn<FTypedElementRowReferenceColumn>(Row);
 	if (const FTypedElementLocalTransformColumn* TransformColumn = DataStorage->GetColumn<FTypedElementLocalTransformColumn>(RefColumn.Row))
 	{
-		UpdateTransformHeadsUpDisplay(*DataStorage, *DataStorage->GetColumn<FTypedElementSlateWidgetReferenceColumn>(Row), InternalTransformHelpers::GetAbnormalTransformTypes(TransformColumn->Transform));
+		UpdateTransformHeadsUpDisplay(*DataStorage->GetColumn<FTypedElementSlateWidgetReferenceColumn>(Row), 
+			InternalTransformHelpers::GetAbnormalTransformTypes(TransformColumn->Transform));
 	}
 
 	return true;
diff --git a/Engine/Plugins/Experimental/TypedElementsDataStorage/Source/TypedElementsDataStorageUI/Public/TypedElementExportedTextWidget.h b/Engine/Plugins/Experimental/TypedElementsDataStorage/Source/TypedElementsDataStorageUI/Public/TypedElementExportedTextWidget.h
--- a/Engine/Plugins/Experimental/TypedElementsDataStorage/Source/TypedElementsDataStorageUI/Public/TypedElementExportedTextWidget.h
+++ b/Engine/Plugins/Experimental/TypedElementsDataStorage/Source/TypedElementsDataStorageUI/Public/TypedElementExportedTextWidget.h
@@ -18,9 +18,10 @@ class TYPEDELEMENTSDATASTORAGEUI_API UTypedElementExportedTextWidgetFactory : pu
 public:
 	~UTypedElementExportedTextWidgetFactory() override = default;
 
-	void RegisterQueries(ITypedElementDataStorageInterface& DataStorage) const override;
 	void RegisterWidgetConstructors(ITypedElementDataStorageInterface& DataStorage,
 		ITypedElementDataStorageUiInterface& DataStorageUi) const override;
+
+	TSet<TWeakObjectPtr<const UScriptStruct>> RegisteredTypes;
 };
 
 USTRUCT()
