[TEDS] Fix an issue where transform column in scene outliner was not displaying tooltips on hover

diff --git a/Engine/Plugins/Experimental/TypedElementsDataStorage/Source/TypedElementsDataStorageUI/Private/TypedElementTransformHeadsUpWidget.cpp b/Engine/Plugins/Experimental/TypedElementsDataStorage/Source/TypedElementsDataStorageUI/Private/TypedElementTransformHeadsUpWidget.cpp
--- a/Engine/Plugins/Experimental/TypedElementsDataStorage/Source/TypedElementsDataStorageUI/Private/TypedElementTransformHeadsUpWidget.cpp
+++ b/Engine/Plugins/Experimental/TypedElementsDataStorage/Source/TypedElementsDataStorageUI/Private/TypedElementTransformHeadsUpWidget.cpp
@@ -75,7 +75,7 @@ public:
 			SNew(SImage)
 				.Image(FAppStyle::GetBrush("EditorViewport.ScaleGridSnap"))
 				.ToolTipText(LOCTEXT("NonUniformScaleTooltip", "Has Non-Uniform Scale"))
-				.Visibility_Lambda([this]() { return EnumHasAllFlags(AbnormalTransformFlags, EAbnormalTransformTypes::NonUniformScale) ? EVisibility::HitTestInvisible : EVisibility::Collapsed; })
+				.Visibility_Lambda([this]() { return EnumHasAllFlags(AbnormalTransformFlags, EAbnormalTransformTypes::NonUniformScale) ? EVisibility::Visible : EVisibility::Collapsed; })
 		];
 		AddSlot().AutoWidth().Padding(IconPadding)
 		[
@@ -83,7 +83,7 @@ public:
 				.Image(FAppStyle::GetBrush("EditorViewport.ScaleMode"))
 				.ColorAndOpacity(SNumericEntryBox<double>::RedLabelBackgroundColor)
 				.ToolTipText(LOCTEXT("NegativeXScaleTooltip", "Has Negative X Scale"))
-				.Visibility_Lambda([this]() { return EnumHasAllFlags(AbnormalTransformFlags, EAbnormalTransformTypes::NegativeXScale) ? EVisibility::HitTestInvisible : EVisibility::Collapsed; })
+				.Visibility_Lambda([this]() { return EnumHasAllFlags(AbnormalTransformFlags, EAbnormalTransformTypes::NegativeXScale) ? EVisibility::Visible : EVisibility::Collapsed; })
 		];
 		AddSlot().AutoWidth().Padding(IconPadding)
 		[
@@ -91,7 +91,7 @@ public:
 				.Image(FAppStyle::GetBrush("EditorViewport.ScaleMode"))
 				.ColorAndOpacity(SNumericEntryBox<double>::GreenLabelBackgroundColor)
 				.ToolTipText(LOCTEXT("NegativeYScaleTooltip", "Has Negative Y Scale"))
-				.Visibility_Lambda([this]() { return EnumHasAllFlags(AbnormalTransformFlags, EAbnormalTransformTypes::NegativeYScale) ? EVisibility::HitTestInvisible : EVisibility::Collapsed; })
+				.Visibility_Lambda([this]() { return EnumHasAllFlags(AbnormalTransformFlags, EAbnormalTransformTypes::NegativeYScale) ? EVisibility::Visible : EVisibility::Collapsed; })
 		];
 		AddSlot().AutoWidth().Padding(IconPadding)
 		[
@@ -99,21 +99,21 @@ public:
 				.Image(FAppStyle::GetBrush("EditorViewport.ScaleMode"))
 				.ColorAndOpacity(SNumericEntryBox<double>::BlueLabelBackgroundColor)
 				.ToolTipText(LOCTEXT("NegativeZScaleTooltip", "Has Negative Z Scale"))
-				.Visibility_Lambda([this]() { return EnumHasAllFlags(AbnormalTransformFlags, EAbnormalTransformTypes::NegativeZScale) ? EVisibility::HitTestInvisible : EVisibility::Collapsed; })
+				.Visibility_Lambda([this]() { return EnumHasAllFlags(AbnormalTransformFlags, EAbnormalTransformTypes::NegativeZScale) ? EVisibility::Visible : EVisibility::Collapsed; })
 		];
 		AddSlot().AutoWidth().Padding(IconPadding)
 		[
 			SNew(SImage)
 				.Image(FAppStyle::GetBrush("SurfaceDetails.AntiClockwiseRotation"))
 				.ToolTipText(LOCTEXT("UnnormalizedRotationTooltip", "Has Un-normalized Rotation"))
-				.Visibility_Lambda([this]() { return EnumHasAllFlags(AbnormalTransformFlags, EAbnormalTransformTypes::UnnormalizedRotation) ? EVisibility::HitTestInvisible : EVisibility::Collapsed; })
+				.Visibility_Lambda([this]() { return EnumHasAllFlags(AbnormalTransformFlags, EAbnormalTransformTypes::UnnormalizedRotation) ? EVisibility::Visible : EVisibility::Collapsed; })
 		];
 		AddSlot().AutoWidth().Padding(IconPadding)
 		[
 			SNew(SImage)
 				.Image(FAppStyle::GetBrush("Symbols.Check"))
 				.ToolTipText(LOCTEXT("NothingToReportTooltip", "No Abnormal Transform Data"))
-				.Visibility_Lambda([this]() { return (AbnormalTransformFlags == EAbnormalTransformTypes::None) ? EVisibility::HitTestInvisible : EVisibility::Collapsed; })
+				.Visibility_Lambda([this]() { return (AbnormalTransformFlags == EAbnormalTransformTypes::None) ? EVisibility::Visible : EVisibility::Collapsed; })
 		];
 	}
 
