CIS fix after recent Mass's CreateArchetype deprecation

diff --git a/Engine/Plugins/Experimental/TypedElementsDataStorage/Source/TypedElementsDataStorage/Private/TypedElementDatabase.cpp b/Engine/Plugins/Experimental/TypedElementsDataStorage/Source/TypedElementsDataStorage/Private/TypedElementDatabase.cpp
--- a/Engine/Plugins/Experimental/TypedElementsDataStorage/Source/TypedElementsDataStorage/Private/TypedElementDatabase.cpp
+++ b/Engine/Plugins/Experimental/TypedElementsDataStorage/Source/TypedElementsDataStorage/Private/TypedElementDatabase.cpp
@@ -127,7 +127,9 @@ TypedElementTableHandle UTypedElementDatabase::RegisterTable(TConstArrayView<con
 	if (ActiveEditorEntityManager && !TableNameLookup.Contains(Name))
 	{
 		TypedElementTableHandle Result = Tables.Num();
-		Tables.Add(ActiveEditorEntityManager->CreateArchetype(ColumnList, Name));
+		FMassArchetypeCreationParams ArchetypeCreationParams;
+		ArchetypeCreationParams.DebugName = Name;
+		Tables.Add(ActiveEditorEntityManager->CreateArchetype(ColumnList, ArchetypeCreationParams));
 		if (Name.IsValid())
 		{
 			TableNameLookup.Add(Name, Result);
@@ -143,7 +145,9 @@ TypedElementTableHandle UTypedElementDatabase::RegisterTable(TypedElementTableHa
 	if (ActiveEditorEntityManager && SourceTable < Tables.Num() && !TableNameLookup.Contains(Name))
 	{
 		TypedElementTableHandle Result = Tables.Num();
-		Tables.Add(ActiveEditorEntityManager->CreateArchetype(Tables[SourceTable], ColumnList, Name));
+		FMassArchetypeCreationParams ArchetypeCreationParams;
+		ArchetypeCreationParams.DebugName = Name;
+		Tables.Add(ActiveEditorEntityManager->CreateArchetype(Tables[SourceTable], ColumnList, ArchetypeCreationParams));
 		if (Name.IsValid())
 		{
 			TableNameLookup.Add(Name, Result);
