Addressed several issues with UObject support in TEDS compatibility layer.

The following issues were addressed:
- AddCompatibleObject, RemoveCompatibleObject and FindRowWithCompatibleObject would try to copy an object if passed by reference. They will now correctly use references.
- UObjects now get tagged to update in the same way actors already did.
- When objects are modified they now also get tagged to update. This helped catch changes that are transactioned but are not explicit properties.
- A set instead of an array is now used to store the objects pending the addition of tags. This avoids the need to sort and deduplicate later on.
- The standard UObject table now corretly adds the FTypedElementPackagePathColumn and FTypedElementPackageLoadedPathColumn. The query callback already supported getting package information from UObjects, but this was never added to the standard table.

diff --git a/Engine/Plugins/Experimental/TypedElementsDataStorage/Source/TypedElementsDataStorage/Private/TypedElementDatabaseCompatibility.cpp b/Engine/Plugins/Experimental/TypedElementsDataStorage/Source/TypedElementsDataStorage/Private/TypedElementDatabaseCompatibility.cpp
--- a/Engine/Plugins/Experimental/TypedElementsDataStorage/Source/TypedElementsDataStorage/Private/TypedElementDatabaseCompatibility.cpp
+++ b/Engine/Plugins/Experimental/TypedElementsDataStorage/Source/TypedElementsDataStorage/Private/TypedElementDatabaseCompatibility.cpp
@@ -24,10 +24,12 @@ void UTypedElementDatabaseCompatibility::Initialize(ITypedElementDataStorageInte
 	StorageInterface->OnUpdate().AddUObject(this, &UTypedElementDatabaseCompatibility::Tick);
 
 	PostEditChangePropertyDelegateHandle = FCoreUObjectDelegates::OnObjectPropertyChanged.AddUObject(this, &UTypedElementDatabaseCompatibility::OnPostEditChangeProperty);
+	ObjectModifiedDelegateHandle = FCoreUObjectDelegates::OnObjectModified.AddUObject(this, &UTypedElementDatabaseCompatibility::OnObjectModified);
 }
 
 void UTypedElementDatabaseCompatibility::Deinitialize()
 {
+	FCoreUObjectDelegates::OnObjectPropertyChanged.Remove(ObjectModifiedDelegateHandle);
 	FCoreUObjectDelegates::OnObjectPropertyChanged.Remove(PostEditChangePropertyDelegateHandle);
 	
 	Reset();
@@ -227,6 +229,7 @@ void UTypedElementDatabaseCompatibility::CreateStandardArchetypes()
 
 	StandardUObjectTable = Storage->RegisterTable(TTypedElementColumnTypeList<
 			FTypedElementUObjectColumn, FTypedElementClassTypeInfoColumn,
+			FTypedElementPackagePathColumn, FTypedElementPackageLoadedPathColumn,
 			FTypedElementSyncFromWorldTag>(), 
 		FName("Editor_StandardUObjectTable"));
 
@@ -247,7 +250,7 @@ void UTypedElementDatabaseCompatibility::Tick()
 		TickPendingActorRegistration(EditorWorld);
 		TickPendingUObjectRegistration();
 		TickPendingExternalObjectRegistration();
-		TickActorSync();
+		TickObjectSync();
 	}
 }
 
@@ -450,29 +453,22 @@ void UTypedElementDatabaseCompatibility::TickPendingExternalObjectRegistration()
 	}
 }
 
-void UTypedElementDatabaseCompatibility::TickActorSync()
+void UTypedElementDatabaseCompatibility::TickObjectSync()
 {
-	if (!ActorsNeedingFullSync.IsEmpty())
+	if (!ObjectsNeedingFullSync.IsEmpty())
 	{
-		TEDS_EVENT_SCOPE(TEXT("Process ActorsNeedingFullSync"));
-		// Deduplicate to avoid duplicate reverse lookups or adding tags more than once
-		{
-			TEDS_EVENT_SCOPE(TEXT("Deduplicate ActorsNeedingFullSync"));
-			ActorsNeedingFullSync.Sort();
-			ActorsNeedingFullSync.SetNum(Algo::Unique(ActorsNeedingFullSync), /* bAllowShrinking */ false);
-		}
-
+		TEDS_EVENT_SCOPE(TEXT("Process ObjectsNeedingFullSync"));
+		
 		TArray<TypedElementRowHandle> RowHandles;
 		{
-			TEDS_EVENT_SCOPE(TEXT("Reverse lookup Rows from Actors"));
+			TEDS_EVENT_SCOPE(TEXT("Reverse lookup Rows from Objects"));
 
-			RowHandles.SetNumUninitialized(ActorsNeedingFullSync.Num());
+			RowHandles.SetNumUninitialized(ObjectsNeedingFullSync.Num());
 			{
 				int32 RowHandleIndex = 0;
-				for (int32 ActorIndex = 0, End = ActorsNeedingFullSync.Num(); ActorIndex < End; ++ActorIndex)
+				for (TObjectKey<const UObject> ObjectKey : ObjectsNeedingFullSync)
 				{
-					const TObjectKey<const AActor> ActorKey = ActorsNeedingFullSync[ActorIndex];
-					const TypedElementRowHandle Row = FindRowWithCompatibleObject(ActorKey);
+					const TypedElementRowHandle Row = FindRowWithCompatibleObject(ObjectKey);
 					if (Row != TypedElementInvalidRowHandle)
 					{
 						RowHandles[RowHandleIndex++] = Row;
@@ -483,7 +479,7 @@ void UTypedElementDatabaseCompatibility::TickActorSync()
 				RowHandles.SetNum(RowHandleCount, bAllowShrinking);
 			}
 
-			ActorsNeedingFullSync.Reset();
+			ObjectsNeedingFullSync.Reset();
 		}
 
 		{
@@ -502,14 +498,16 @@ void UTypedElementDatabaseCompatibility::OnPostEditChangeProperty(
 	UObject* Object,
 	FPropertyChangedEvent& /*PropertyChangedEvent*/)
 {
-	// We aren't sure if this Actor is tracked by the database
-	// Will resolve that during the tick step
-	// The aim is to keep this delegate handler as simple as possible to avoid performance side-effects when other code
-	// invokes it.
-	AActor* Actor = Cast<AActor>(Object);
-	if (Actor)
-	{
-		// Note: This array may end up with duplicates
-		ActorsNeedingFullSync.Add(Actor);
-	}
+	// Determining the object is being tracked in the database can't be done safely as it may be queued for addition.
+	// It would also add a small bit of performance overhead as access the lookup table can be done faster as a
+	// batch operation during the tick step.
+	ObjectsNeedingFullSync.FindOrAdd(Object);
+}
+
+void UTypedElementDatabaseCompatibility::OnObjectModified(UObject* Object)
+{
+	// Determining the object is being tracked in the database can't be done safely as it may be queued for addition.
+	// It would also add a small bit of performance overhead as access the lookup table can be done faster as a
+	// batch operation during the tick step.
+	ObjectsNeedingFullSync.FindOrAdd(Object);
 }
diff --git a/Engine/Plugins/Experimental/TypedElementsDataStorage/Source/TypedElementsDataStorage/Private/TypedElementDatabaseCompatibility.h b/Engine/Plugins/Experimental/TypedElementsDataStorage/Source/TypedElementsDataStorage/Private/TypedElementDatabaseCompatibility.h
--- a/Engine/Plugins/Experimental/TypedElementsDataStorage/Source/TypedElementsDataStorage/Private/TypedElementDatabaseCompatibility.h
+++ b/Engine/Plugins/Experimental/TypedElementsDataStorage/Source/TypedElementsDataStorage/Private/TypedElementDatabaseCompatibility.h
@@ -50,9 +50,10 @@ private:
 	void TickPendingActorRegistration(UWorld* EditorWorld);
 	void TickPendingUObjectRegistration();
 	void TickPendingExternalObjectRegistration();
-	void TickActorSync();
+	void TickObjectSync();
 
 	void OnPostEditChangeProperty(UObject* Object, FPropertyChangedEvent& PropertyChangedEvent);
+	void OnObjectModified(UObject* Object);
 	
 	template<typename AddressType>
 	struct PendingRegistration
@@ -92,11 +93,11 @@ private:
 	TMap<void*, TypedElementRowHandle> ReverseObjectLookup;
 
 	/**
-	 * Reference of actors that need to be fully synced from the world to the database.
-	 * May have duplicates
-	 * Caution: Could point to actors that have been GC-ed
+	 * Reference of objects (UObject and AActor) that need to be fully synced from the world to the database.
+	 * Caution: Could point to objects that have been GC-ed
 	 */
-	TArray<TObjectKey<const AActor>> ActorsNeedingFullSync;
+	TSet<TObjectKey<const UObject>> ObjectsNeedingFullSync;
 
 	FDelegateHandle PostEditChangePropertyDelegateHandle;
+	FDelegateHandle ObjectModifiedDelegateHandle;
 };
\ No newline at end of file
diff --git a/Engine/Source/Runtime/TypedElementFramework/Public/Elements/Interfaces/TypedElementDataStorageCompatibilityInterface.h b/Engine/Source/Runtime/TypedElementFramework/Public/Elements/Interfaces/TypedElementDataStorageCompatibilityInterface.h
--- a/Engine/Source/Runtime/TypedElementFramework/Public/Elements/Interfaces/TypedElementDataStorageCompatibilityInterface.h
+++ b/Engine/Source/Runtime/TypedElementFramework/Public/Elements/Interfaces/TypedElementDataStorageCompatibilityInterface.h
@@ -42,16 +42,16 @@ public:
 	 * with a row and to setup the initial row data.
 	 */
 	template<typename ObjectType>
-	TypedElementRowHandle AddCompatibleObject(ObjectType Object);
+	TypedElementRowHandle AddCompatibleObject(ObjectType&& Object);
 	template<typename ObjectType>
-	TypedElementRowHandle AddCompatibleObject(ObjectType Object, TypedElementTableHandle Table);
+	TypedElementRowHandle AddCompatibleObject(ObjectType&& Object, TypedElementTableHandle Table);
 
 	/** Removes a previously registered object from the data storage. */
 	template<typename ObjectType>
-	void RemoveCompatibleObject(ObjectType Object);
+	void RemoveCompatibleObject(ObjectType&& Object);
 
 	template<typename ObjectType>
-	TypedElementRowHandle FindRowWithCompatibleObject(const ObjectType Object) const;
+	TypedElementRowHandle FindRowWithCompatibleObject(ObjectType&& Object) const;
 
 
 	/**
@@ -98,12 +98,12 @@ template<typename Type> Type* GetRawPointer(Type* Object)						{ return Object;
 template<typename Type> Type* GetRawPointer(Type& Object)						{ return &Object; }
 
 template<typename ObjectType>
-TypedElementRowHandle ITypedElementDataStorageCompatibilityInterface::AddCompatibleObject(ObjectType Object)
+TypedElementRowHandle ITypedElementDataStorageCompatibilityInterface::AddCompatibleObject(ObjectType&& Object)
 {
-	auto RawPointer = GetRawPointer(Object);
+	auto RawPointer = GetRawPointer(Forward<ObjectType>(Object));
 	using BaseType = std::remove_cv_t<std::remove_pointer_t<decltype(RawPointer)>>;
 
-	if constexpr (std::is_same_v<BaseType, AActor> || std::is_same_v<BaseType, UObject>)
+	if constexpr (std::is_base_of_v<UObject, BaseType>)
 	{
 		return AddCompatibleObjectExplicit(RawPointer);
 	}
@@ -114,31 +114,38 @@ TypedElementRowHandle ITypedElementDataStorageCompatibilityInterface::AddCompati
 }
 
 template<typename ObjectType>
-TypedElementRowHandle ITypedElementDataStorageCompatibilityInterface::AddCompatibleObject(ObjectType Object, TypedElementTableHandle Table)
+TypedElementRowHandle ITypedElementDataStorageCompatibilityInterface::AddCompatibleObject(ObjectType&& Object, TypedElementTableHandle Table)
 {
-	auto RawPointer = GetRawPointer(Object);
-	using BaseType = std::remove_cv_t<std::remove_pointer_t<decltype(RawPointer)>>;
-
-	if constexpr (std::is_same_v<BaseType, AActor> || std::is_same_v<BaseType, UObject>)
+	if (Table != TypedElementInvalidTableHandle)
 	{
-		return AddCompatibleObjectExplicit(RawPointer, Table);
+		auto RawPointer = GetRawPointer(Forward<ObjectType>(Object));
+		using BaseType = std::remove_cv_t<std::remove_pointer_t<decltype(RawPointer)>>;
+
+		if constexpr (std::is_base_of_v<UObject, BaseType>)
+		{
+			return AddCompatibleObjectExplicit(RawPointer, Table);
+		}
+		else
+		{
+			return AddCompatibleObjectExplicit(RawPointer, BaseType::StaticStruct(), Table);
+		}
 	}
 	else
 	{
-		return AddCompatibleObjectExplicit(RawPointer, BaseType::StaticStruct(), Table);
+		return AddCompatibleObject(Object);
 	}
 }
 
 template<typename ObjectType>
-void ITypedElementDataStorageCompatibilityInterface::RemoveCompatibleObject(ObjectType Object)
+void ITypedElementDataStorageCompatibilityInterface::RemoveCompatibleObject(ObjectType&& Object)
 {
-	RemoveCompatibleObjectExplicit(GetRawPointer(Object));
+	RemoveCompatibleObjectExplicit(GetRawPointer(Forward<ObjectType>(Object)));
 }
 
 template<typename ObjectType>
-TypedElementRowHandle ITypedElementDataStorageCompatibilityInterface::FindRowWithCompatibleObject(const ObjectType Object) const
+TypedElementRowHandle ITypedElementDataStorageCompatibilityInterface::FindRowWithCompatibleObject(ObjectType&& Object) const
 {
-	return FindRowWithCompatibleObjectExplicit(GetRawPointer(Object));
+	return FindRowWithCompatibleObjectExplicit(GetRawPointer(Forward<ObjectType>(Object)));
 }
 
 template<typename Subsystem>
