Added ability to TEDS to store variable sized data from a query callback.

So far the query callbacks weren't able to store any additional data on the delayed commands. This prevented operations like creating a new column and assigning it data in the same call. It also caused a performance loss as every multi-row and/or multi-column call was broken down into individual pieces behind the scenes. This changes introduces a scratch buffer that's conceptually alive for a single frame and allows these types of temporary data to be stored as the delayed commands are going to be executed by the end of the frame.

The following notable changes were made:
- Columns can now be added with data. The data is stored in temporary location and copied during the end of a phase, or moved if a move operator callback is provided.
- A thread-safe scratch buffer was added that allows for allocations that are cleared at the end of the frame.
- The scratch buffer is setup to recycle any full blocks once per frame. This includes deleting all stored objects that require doing so.
- The scratch buffer is available for calls that implement the query callback context.
- RemoveRows in the query callback context uses the Mass native "DestroyEntities"-call as TEDS' row handle and Mass' entity handle are an identical size and directly map to eachother.
- AddColumns and RemoveColumns in the query callback context now use the calls to update the composition of an entity instead of multiple calls to update the columns individually. This is a bit more expensive to set up on the callsite, but requires no/less additional memory and is a lot faster to execute at the end of a phase.
- A bug was fixed whereby deleting a row using an invalid row handle would assert.

diff --git a/Engine/Plugins/Experimental/TypedElementsDataStorage/Source/TypedElementsDataStorage/Private/Processors/TypedElementProcessorAdaptors.cpp b/Engine/Plugins/Experimental/TypedElementsDataStorage/Source/TypedElementsDataStorage/Private/Processors/TypedElementProcessorAdaptors.cpp
--- a/Engine/Plugins/Experimental/TypedElementsDataStorage/Source/TypedElementsDataStorage/Private/Processors/TypedElementProcessorAdaptors.cpp
+++ b/Engine/Plugins/Experimental/TypedElementsDataStorage/Source/TypedElementsDataStorage/Private/Processors/TypedElementProcessorAdaptors.cpp
@@ -9,6 +9,7 @@
 #include "MassExecutionContext.h"
 #include "Queries/TypedElementExtendedQueryStore.h"
 #include "TypedElementDatabase.h"
+#include "TypedElementDatabaseScratchBuffer.h"
 
 template<typename T>
 struct FMassContextCommon : public T
@@ -148,11 +149,32 @@ private:
 
 struct FMassContextForwarder final : public FMassContextCommon<ITypedElementDataStorageInterface::IQueryContext>
 {
+private:
+	void TedsColumnsToMassDescriptor(FMassArchetypeCompositionDescriptor& Descriptor, TConstArrayView<const UScriptStruct*> ColumnTypes)
+	{
+		for (const UScriptStruct* ColumnType : ColumnTypes)
+		{
+			if (ColumnType->IsChildOf(FMassTag::StaticStruct()))
+			{
+				Descriptor.Tags.Add(*ColumnType);
+			}
+			else
+			{
+				checkf(ColumnType->IsChildOf(FMassFragment::StaticStruct()),
+					TEXT("Given struct type is not a valid fragment or tag type."));
+				Descriptor.Fragments.Add(*ColumnType);
+
+			}
+		}
+	}
+public:
+
 	FMassContextForwarder(ITypedElementDataStorageInterface::FQueryDescription& InQueryDescription, FMassExecutionContext& InContext, 
-		FTypedElementExtendedQueryStore& InQueryStore)
+		FTypedElementExtendedQueryStore& InQueryStore, FTypedElementDatabaseScratchBuffer& InScratchBuffer)
 		: FMassContextCommon(InContext)
 		, QueryDescription(InQueryDescription)
 		, QueryStore(InQueryStore)
+		, ScratchBuffer(InScratchBuffer)
 	{}
 
 	~FMassContextForwarder() override = default;
@@ -199,82 +221,206 @@ struct FMassContextForwarder final : public FMassContextCommon<ITypedElementData
 
 	void RemoveRows(TConstArrayView<TypedElementRowHandle> Rows) override
 	{
-		for (TypedElementRowHandle Row : Rows)
-		{
-			Context.Defer().DestroyEntity(FMassEntityHandle::FromNumber(Row));
-		}
+		// Row handles and entities map 1:1 for data, so a reintpret_cast can be safely done to avoid
+		// having to allocate memory and iterating over the rows.
+
+		static_assert(sizeof(FMassEntityHandle) == sizeof(TypedElementRowHandle), 
+			"Sizes of mass entity and data storage row have gone out of sync.");
+		static_assert(alignof(FMassEntityHandle) == alignof(TypedElementRowHandle),
+			"Alignment of mass entity and data storage row have gone out of sync.");
+
+		Context.Defer().DestroyEntities(
+			TConstArrayView<FMassEntityHandle>(reinterpret_cast<const FMassEntityHandle*>(Rows.begin()), Rows.Num()));
 	}
 
-	void AddColumns(TypedElementRowHandle Row, TConstArrayView<const UScriptStruct*> ColumnTypes) override
+	void* AddColumnUnitialized(TypedElementDataStorage::RowHandle Row, const UScriptStruct* ObjectType) override
 	{
-		for (const UScriptStruct* ColumnType : ColumnTypes)
+		struct FAddValueColumn
 		{
-			bool bIsTag = ColumnType->IsChildOf(FMassTag::StaticStruct());
+			const UScriptStruct* FragmentType;
+			FMassEntityHandle Entity;
+			void* Object;
+		};
 
-			checkf(ColumnType->IsChildOf(FMassFragment::StaticStruct()) || bIsTag,
-				TEXT("Given struct type to add is not a valid fragment or tag type."));
+		struct FAddValueColumnWithDestructor : FAddValueColumn
+		{
+			FAddValueColumnWithDestructor(const UScriptStruct* InFragmentType, FMassEntityHandle InEntity, void* InObject)
+				: FAddValueColumn(InFragmentType, InEntity, InObject)
+			{}
 
-			if (bIsTag)
+			~FAddValueColumnWithDestructor()
 			{
-				Context.Defer().PushCommand<FMassDeferredAddCommand>(
-					[EntityHandle = FMassEntityHandle::FromNumber(Row), ColumnType](FMassEntityManager& System)
-					{
-						System.AddTagToEntity(EntityHandle, ColumnType);
-					});
+				FragmentType->DestroyStruct(Object);
 			}
-			else
+		};
+
+		void* ObjectCopy = ScratchBuffer.Allocate(ObjectType->GetStructureSize(), ObjectType->GetMinAlignment());
+		FAddValueColumn* AddedColumn = ObjectType->GetCppStructOps()->HasDestructor() ?
+			ScratchBuffer.Emplace<FAddValueColumnWithDestructor>(ObjectType, FMassEntityHandle::FromNumber(Row), ObjectCopy) :
+			ScratchBuffer.Emplace<FAddValueColumn>(ObjectType, FMassEntityHandle::FromNumber(Row), ObjectCopy);
+
+		Context.Defer().PushCommand<FMassDeferredAddCommand>(
+			[AddedColumn](FMassEntityManager& System)
 			{
-				Context.Defer().PushCommand<FMassDeferredAddCommand>(
-					[EntityHandle = FMassEntityHandle::FromNumber(Row), ColumnType](FMassEntityManager& System)
-					{
-						System.AddFragmentToEntity(EntityHandle, ColumnType);
-					});
+				System.AddFragmentToEntity(AddedColumn->Entity, AddedColumn->FragmentType);
+			});
+		
+		Context.Defer().PushCommand<FMassDeferredSetCommand>(
+			[AddedColumn](FMassEntityManager& System)
+			{
+				FStructView Fragment = System.GetFragmentDataStruct(AddedColumn->Entity, AddedColumn->FragmentType);
+				AddedColumn->FragmentType->CopyScriptStruct(Fragment.GetMemory(), AddedColumn->Object);
+			});
+
+		return ObjectCopy;
+	}
+
+	void* AddColumnUnitialized(TypedElementDataStorage::RowHandle Row, const UScriptStruct* ObjectType, ObjectMoveOperator Mover) override
+	{
+		struct FAddMoveableValueColumn
+		{
+			ObjectMoveOperator Mover;
+			const UScriptStruct* FragmentType;
+			FMassEntityHandle Entity;
+			void* Object;
+		};
+
+		struct FAddMoveableValueColumnWithDestructor : FAddMoveableValueColumn
+		{
+			FAddMoveableValueColumnWithDestructor(
+				ObjectMoveOperator InMover, const UScriptStruct* InFragmentType, FMassEntityHandle InEntity, void* InObject)
+				: FAddMoveableValueColumn(InMover, InFragmentType, InEntity, InObject)
+			{}
+
+			~FAddMoveableValueColumnWithDestructor()
+			{
+				FragmentType->DestroyStruct(Object);
 			}
-		}
+		};
+
+		void* MovedObject = ScratchBuffer.Allocate(ObjectType->GetStructureSize(), ObjectType->GetMinAlignment());
+		FAddMoveableValueColumn* AddedColumn = ObjectType->GetCppStructOps()->HasDestructor() ?
+			ScratchBuffer.Emplace<FAddMoveableValueColumnWithDestructor>(Mover, ObjectType, FMassEntityHandle::FromNumber(Row), MovedObject) :
+			ScratchBuffer.Emplace<FAddMoveableValueColumn>(Mover, ObjectType, FMassEntityHandle::FromNumber(Row), MovedObject);
+
+		Context.Defer().PushCommand<FMassDeferredAddCommand>(
+			[AddedColumn](FMassEntityManager& System)
+			{
+				System.AddFragmentToEntity(AddedColumn->Entity, AddedColumn->FragmentType);
+			});
+
+		Context.Defer().PushCommand<FMassDeferredSetCommand>(
+			[AddedColumn](FMassEntityManager& System)
+			{
+				FStructView Fragment = System.GetFragmentDataStruct(AddedColumn->Entity, AddedColumn->FragmentType);
+				AddedColumn->Mover(Fragment.GetMemory(), AddedColumn->Object);
+			});
+
+		return MovedObject;
+	}
+
+	void AddColumns(TypedElementDataStorage::RowHandle Row, TConstArrayView<const UScriptStruct*> ColumnTypes) override
+	{
+		struct FAddedColumns
+		{
+			FMassArchetypeCompositionDescriptor AddDescriptor;
+			FMassEntityHandle Entity;
+		};
+
+		FAddedColumns* AddedColumns = ScratchBuffer.Emplace<FAddedColumns>();
+		TedsColumnsToMassDescriptor(AddedColumns->AddDescriptor, ColumnTypes);
+		AddedColumns->Entity = FMassEntityHandle::FromNumber(Row);
+
+		Context.Defer().PushCommand<FMassDeferredAddCommand>(
+			[AddedColumns](FMassEntityManager& System)
+			{
+				System.AddCompositionToEntity_GetDelta(AddedColumns->Entity, AddedColumns->AddDescriptor);
+			});
 	}
 
-	void AddColumns(TConstArrayView<TypedElementRowHandle> Rows, TConstArrayView<const UScriptStruct*> ColumnTypes) override
+	void AddColumns(TConstArrayView<TypedElementDataStorage::RowHandle> Rows, TConstArrayView<const UScriptStruct*> ColumnTypes) override
 	{
+		struct FAddedColumns
+		{
+			FMassArchetypeCompositionDescriptor AddDescriptor;
+			FMassEntityHandle* Entities;
+			int32 EntityCount;
+		};
+
+		FAddedColumns* AddedColumns = ScratchBuffer.Emplace<FAddedColumns>();
+		TedsColumnsToMassDescriptor(AddedColumns->AddDescriptor, ColumnTypes);
+		
+		FMassEntityHandle* Entities = ScratchBuffer.EmplaceArray<FMassEntityHandle>(Rows.Num());
+		AddedColumns->Entities = Entities;
 		for (TypedElementRowHandle Row : Rows)
 		{
-			AddColumns(Row, ColumnTypes);
+			*Entities = FMassEntityHandle::FromNumber(Row);
+			Entities++;
 		}
+		AddedColumns->EntityCount = Rows.Num();
+
+		Context.Defer().PushCommand<FMassDeferredAddCommand>(
+			[AddedColumns](FMassEntityManager& System)
+			{
+				FMassEntityHandle* Entities = AddedColumns->Entities;
+				int32 Count = AddedColumns->EntityCount;
+				for (int32 Counter = 0; Counter < Count; ++Counter)
+				{
+					System.AddCompositionToEntity_GetDelta(*Entities++, AddedColumns->AddDescriptor);
+				}
+			});
 	}
 
-	void RemoveColumns(TypedElementRowHandle Row, TConstArrayView<const UScriptStruct*> ColumnTypes) override
+	void RemoveColumns(TypedElementDataStorage::RowHandle Row, TConstArrayView<const UScriptStruct*> ColumnTypes) override
 	{
-		for (const UScriptStruct* ColumnType : ColumnTypes)
+		struct FRemovedColumns
 		{
-			bool bIsTag = ColumnType->IsChildOf(FMassTag::StaticStruct());
+			FMassArchetypeCompositionDescriptor RemoveDescriptor;
+			FMassEntityHandle Entity;
+		};
 
-			checkf(ColumnType->IsChildOf(FMassFragment::StaticStruct()) || bIsTag,
-				TEXT("Given struct type to remove is not a valid fragment or tag type."));
+		FRemovedColumns* RemovedColumns = ScratchBuffer.Emplace<FRemovedColumns>();
+		TedsColumnsToMassDescriptor(RemovedColumns->RemoveDescriptor, ColumnTypes);
+		RemovedColumns->Entity = FMassEntityHandle::FromNumber(Row);
 
-			if (bIsTag)
+		Context.Defer().PushCommand<FMassDeferredAddCommand>(
+			[RemovedColumns](FMassEntityManager& System)
 			{
-				Context.Defer().PushCommand<FMassDeferredAddCommand>(
-					[EntityHandle = FMassEntityHandle::FromNumber(Row), ColumnType](FMassEntityManager& System)
-					{
-						System.RemoveTagFromEntity(EntityHandle, ColumnType);
-					});
-			}
-			else
-			{
-				Context.Defer().PushCommand<FMassDeferredAddCommand>(
-					[EntityHandle = FMassEntityHandle::FromNumber(Row), ColumnType](FMassEntityManager& System)
-					{
-						System.RemoveFragmentFromEntity(EntityHandle, ColumnType);
-					});
-			}
-		}
+				System.RemoveCompositionFromEntity(RemovedColumns->Entity, RemovedColumns->RemoveDescriptor);
+			});
 	}
 
-	void RemoveColumns(TConstArrayView<TypedElementRowHandle> Rows, TConstArrayView<const UScriptStruct*> ColumnTypes) override
+	void RemoveColumns(TConstArrayView<TypedElementDataStorage::RowHandle> Rows, TConstArrayView<const UScriptStruct*> ColumnTypes) override
 	{
+		struct FRemovedColumns
+		{
+			FMassArchetypeCompositionDescriptor RemoveDescriptor;
+			FMassEntityHandle* Entities;
+			int32 EntityCount;
+		};
+
+		FRemovedColumns* RemovedColumns = ScratchBuffer.Emplace<FRemovedColumns>();
+		TedsColumnsToMassDescriptor(RemovedColumns->RemoveDescriptor, ColumnTypes);
+
+		FMassEntityHandle* Entities = ScratchBuffer.EmplaceArray<FMassEntityHandle>(Rows.Num());
+		RemovedColumns->Entities = Entities;
 		for (TypedElementRowHandle Row : Rows)
 		{
-			RemoveColumns(Row, ColumnTypes);
+			*Entities = FMassEntityHandle::FromNumber(Row);
+			Entities++;
 		}
+		RemovedColumns->EntityCount = Rows.Num();
+
+		Context.Defer().PushCommand<FMassDeferredAddCommand>(
+			[RemovedColumns](FMassEntityManager& System)
+			{
+				FMassEntityHandle* Entities = RemovedColumns->Entities;
+				int32 Count = RemovedColumns->EntityCount;
+				for (int32 Counter = 0; Counter < Count; ++Counter)
+				{
+					System.RemoveCompositionFromEntity(*Entities++, RemovedColumns->RemoveDescriptor);
+				}
+			});
 	}
 
 	TypedElementDataStorage::FQueryResult RunQuery(TypedElementQueryHandle Query) override
@@ -323,6 +469,7 @@ struct FMassContextForwarder final : public FMassContextCommon<ITypedElementData
 
 	ITypedElementDataStorageInterface::FQueryDescription& QueryDescription;
 	FTypedElementExtendedQueryStore& QueryStore;
+	FTypedElementDatabaseScratchBuffer& ScratchBuffer;
 };
 
 
@@ -346,15 +493,16 @@ FPhasePreOrPostAmbleExecutor::~FPhasePreOrPostAmbleExecutor()
 void FPhasePreOrPostAmbleExecutor::ExecuteQuery(
 	ITypedElementDataStorageInterface::FQueryDescription& Description, 
 	FTypedElementExtendedQueryStore& QueryStore, 
+	FTypedElementDatabaseScratchBuffer& ScratchBuffer,
 	FMassEntityQuery& NativeQuery,
 	ITypedElementDataStorageInterface::QueryCallbackRef Callback)
 {
 	NativeQuery.ForEachEntityChunk(Context.GetEntityManagerChecked(), Context,
-		[&Callback, &QueryStore, &Description](FMassExecutionContext& ExecutionContext)
+		[&Callback, &QueryStore, &ScratchBuffer, &Description](FMassExecutionContext& ExecutionContext)
 		{
 			if (FTypedElementQueryProcessorData::PrepareCachedDependenciesOnQuery(Description, ExecutionContext))
 			{
-				FMassContextForwarder QueryContext(Description, ExecutionContext, QueryStore);
+				FMassContextForwarder QueryContext(Description, ExecutionContext, QueryStore, ScratchBuffer);
 				Callback(Description, QueryContext);
 			}
 		}
@@ -371,11 +519,16 @@ FTypedElementQueryProcessorData::FTypedElementQueryProcessorData(UMassProcessor&
 {
 }
 
-bool FTypedElementQueryProcessorData::CommonQueryConfiguration(UMassProcessor& InOwner, FTypedElementExtendedQuery& InQuery, 
-	FTypedElementExtendedQueryStore& InQueryStore, TArrayView<FMassEntityQuery> Subqueries)
+bool FTypedElementQueryProcessorData::CommonQueryConfiguration(
+	UMassProcessor& InOwner, 
+	FTypedElementExtendedQuery& InQuery, 
+	FTypedElementExtendedQueryStore& InQueryStore, 
+	FTypedElementDatabaseScratchBuffer& InScratchBuffer,
+	TArrayView<FMassEntityQuery> Subqueries)
 {
 	ParentQuery = &InQuery;
 	QueryStore = &InQueryStore;
+	ScratchBuffer = &InScratchBuffer;
 
 	if (ensureMsgf(InQuery.Description.Subqueries.Num() <= Subqueries.Num(),
 		TEXT("Provided query has too many (%i) subqueries."), InQuery.Description.Subqueries.Num()))
@@ -544,7 +697,7 @@ void FTypedElementQueryProcessorData::Execute(FMassEntityManager& EntityManager,
 		{
 			if (PrepareCachedDependenciesOnQuery(Description, Context))
 			{
-				FMassContextForwarder QueryContext(Description, Context, *QueryStore);
+				FMassContextForwarder QueryContext(Description, Context, *QueryStore, *ScratchBuffer);
 				Description.Callback.Function(Description, QueryContext);
 			}
 		}
@@ -570,15 +723,20 @@ FMassEntityQuery& UTypedElementQueryProcessorCallbackAdapterProcessorBase::GetQu
 }
 
 bool UTypedElementQueryProcessorCallbackAdapterProcessorBase::ConfigureQueryCallback(
-	FTypedElementExtendedQuery& Query, FTypedElementExtendedQueryStore& QueryStore)
+	FTypedElementExtendedQuery& Query, 
+	FTypedElementExtendedQueryStore& QueryStore, 
+	FTypedElementDatabaseScratchBuffer& ScratchBuffer)
 {
-	return ConfigureQueryCallbackData(Query, QueryStore, {});
+	return ConfigureQueryCallbackData(Query, QueryStore, ScratchBuffer, {});
 }
 
-bool UTypedElementQueryProcessorCallbackAdapterProcessorBase::ConfigureQueryCallbackData(FTypedElementExtendedQuery& Query,
-	FTypedElementExtendedQueryStore& QueryStore,  TArrayView<FMassEntityQuery> Subqueries)
+bool UTypedElementQueryProcessorCallbackAdapterProcessorBase::ConfigureQueryCallbackData(
+	FTypedElementExtendedQuery& Query,
+	FTypedElementExtendedQueryStore& QueryStore, 
+	FTypedElementDatabaseScratchBuffer& ScratchBuffer,
+	TArrayView<FMassEntityQuery> Subqueries)
 {
-	bool Result = Data.CommonQueryConfiguration(*this, Query, QueryStore, Subqueries);
+	bool Result = Data.CommonQueryConfiguration(*this, Query, QueryStore, ScratchBuffer, Subqueries);
 
 	bRequiresGameThreadExecution = Query.Description.Callback.bForceToGameThread;
 	ExecutionFlags = static_cast<int32>(EProcessorExecutionFlags::Editor); 
@@ -614,27 +772,27 @@ void UTypedElementQueryProcessorCallbackAdapterProcessorBase::Execute(FMassEntit
 }
 
 bool UTypedElementQueryProcessorCallbackAdapterProcessorWith1Subquery::ConfigureQueryCallback(
-	FTypedElementExtendedQuery& Query, FTypedElementExtendedQueryStore& QueryStore)
+	FTypedElementExtendedQuery& Query, FTypedElementExtendedQueryStore& QueryStore, FTypedElementDatabaseScratchBuffer& ScratchBuffer)
 {
-	return ConfigureQueryCallbackData(Query, QueryStore, NativeSubqueries);
+	return ConfigureQueryCallbackData(Query, QueryStore, ScratchBuffer, NativeSubqueries);
 }
 
 bool UTypedElementQueryProcessorCallbackAdapterProcessorWith2Subqueries::ConfigureQueryCallback(
-	FTypedElementExtendedQuery& Query, FTypedElementExtendedQueryStore& QueryStore)
+	FTypedElementExtendedQuery& Query, FTypedElementExtendedQueryStore& QueryStore, FTypedElementDatabaseScratchBuffer& ScratchBuffer)
 {
-	return ConfigureQueryCallbackData(Query, QueryStore, NativeSubqueries);
+	return ConfigureQueryCallbackData(Query, QueryStore, ScratchBuffer, NativeSubqueries);
 }
 
 bool UTypedElementQueryProcessorCallbackAdapterProcessorWith3Subqueries::ConfigureQueryCallback(
-	FTypedElementExtendedQuery& Query, FTypedElementExtendedQueryStore& QueryStore)
+	FTypedElementExtendedQuery& Query, FTypedElementExtendedQueryStore& QueryStore, FTypedElementDatabaseScratchBuffer& ScratchBuffer)
 {
-	return ConfigureQueryCallbackData(Query, QueryStore, NativeSubqueries);
+	return ConfigureQueryCallbackData(Query, QueryStore, ScratchBuffer, NativeSubqueries);
 }
 
 bool UTypedElementQueryProcessorCallbackAdapterProcessorWith4Subqueries::ConfigureQueryCallback(
-	FTypedElementExtendedQuery& Query, FTypedElementExtendedQueryStore& QueryStore)
+	FTypedElementExtendedQuery& Query, FTypedElementExtendedQueryStore& QueryStore, FTypedElementDatabaseScratchBuffer& ScratchBuffer)
 {
-	return ConfigureQueryCallbackData(Query, QueryStore, NativeSubqueries);
+	return ConfigureQueryCallbackData(Query, QueryStore, ScratchBuffer, NativeSubqueries);
 }
 
 
@@ -665,15 +823,15 @@ EMassObservedOperation UTypedElementQueryObserverCallbackAdapterProcessorBase::G
 }
 
 bool UTypedElementQueryObserverCallbackAdapterProcessorBase::ConfigureQueryCallback(
-	FTypedElementExtendedQuery& Query, FTypedElementExtendedQueryStore& QueryStore)
+	FTypedElementExtendedQuery& Query, FTypedElementExtendedQueryStore& QueryStore, FTypedElementDatabaseScratchBuffer& ScratchBuffer)
 {
-	return ConfigureQueryCallbackData(Query, QueryStore, {});
+	return ConfigureQueryCallbackData(Query, QueryStore, ScratchBuffer, {});
 }
 
 bool UTypedElementQueryObserverCallbackAdapterProcessorBase::ConfigureQueryCallbackData(FTypedElementExtendedQuery& Query,
-	FTypedElementExtendedQueryStore& QueryStore, TArrayView<FMassEntityQuery> Subqueries)
+	FTypedElementExtendedQueryStore& QueryStore, FTypedElementDatabaseScratchBuffer& ScratchBuffer, TArrayView<FMassEntityQuery> Subqueries)
 {
-	bool Result = Data.CommonQueryConfiguration(*this, Query, QueryStore, Subqueries);
+	bool Result = Data.CommonQueryConfiguration(*this, Query, QueryStore, ScratchBuffer, Subqueries);
 
 	bRequiresGameThreadExecution = Query.Description.Callback.bForceToGameThread;
 	ExecutionFlags = static_cast<int32>(EProcessorExecutionFlags::Editor);
@@ -738,25 +896,25 @@ void UTypedElementQueryObserverCallbackAdapterProcessorBase::Execute(FMassEntity
 }
 
 bool UTypedElementQueryObserverCallbackAdapterProcessorWith1Subquery::ConfigureQueryCallback(
-	FTypedElementExtendedQuery& Query, FTypedElementExtendedQueryStore& QueryStore)
+	FTypedElementExtendedQuery& Query, FTypedElementExtendedQueryStore& QueryStore, FTypedElementDatabaseScratchBuffer& ScratchBuffer)
 {
-	return ConfigureQueryCallbackData(Query, QueryStore, NativeSubqueries);
+	return ConfigureQueryCallbackData(Query, QueryStore, ScratchBuffer, NativeSubqueries);
 }
 
 bool UTypedElementQueryObserverCallbackAdapterProcessorWith2Subqueries::ConfigureQueryCallback(
-	FTypedElementExtendedQuery& Query, FTypedElementExtendedQueryStore& QueryStore)
+	FTypedElementExtendedQuery& Query, FTypedElementExtendedQueryStore& QueryStore, FTypedElementDatabaseScratchBuffer& ScratchBuffer)
 {
-	return ConfigureQueryCallbackData(Query, QueryStore, NativeSubqueries);
+	return ConfigureQueryCallbackData(Query, QueryStore, ScratchBuffer, NativeSubqueries);
 }
 
 bool UTypedElementQueryObserverCallbackAdapterProcessorWith3Subqueries::ConfigureQueryCallback(
-	FTypedElementExtendedQuery& Query, FTypedElementExtendedQueryStore& QueryStore)
+	FTypedElementExtendedQuery& Query, FTypedElementExtendedQueryStore& QueryStore, FTypedElementDatabaseScratchBuffer& ScratchBuffer)
 {
-	return ConfigureQueryCallbackData(Query, QueryStore, NativeSubqueries);
+	return ConfigureQueryCallbackData(Query, QueryStore, ScratchBuffer, NativeSubqueries);
 }
 
 bool UTypedElementQueryObserverCallbackAdapterProcessorWith4Subqueries::ConfigureQueryCallback(
-	FTypedElementExtendedQuery& Query, FTypedElementExtendedQueryStore& QueryStore)
+	FTypedElementExtendedQuery& Query, FTypedElementExtendedQueryStore& QueryStore, FTypedElementDatabaseScratchBuffer& ScratchBuffer)
 {
-	return ConfigureQueryCallbackData(Query, QueryStore, NativeSubqueries);
+	return ConfigureQueryCallbackData(Query, QueryStore, ScratchBuffer, NativeSubqueries);
 }
diff --git a/Engine/Plugins/Experimental/TypedElementsDataStorage/Source/TypedElementsDataStorage/Private/Processors/TypedElementProcessorAdaptors.h b/Engine/Plugins/Experimental/TypedElementsDataStorage/Source/TypedElementsDataStorage/Private/Processors/TypedElementProcessorAdaptors.h
--- a/Engine/Plugins/Experimental/TypedElementsDataStorage/Source/TypedElementsDataStorage/Private/Processors/TypedElementProcessorAdaptors.h
+++ b/Engine/Plugins/Experimental/TypedElementsDataStorage/Source/TypedElementsDataStorage/Private/Processors/TypedElementProcessorAdaptors.h
@@ -11,6 +11,7 @@
 
 #include "TypedElementProcessorAdaptors.generated.h"
 
+class FTypedElementDatabaseScratchBuffer;
 struct FTypedElementExtendedQuery;
 class FTypedElementExtendedQueryStore;
 
@@ -22,6 +23,7 @@ struct FPhasePreOrPostAmbleExecutor
 	void ExecuteQuery(
 		ITypedElementDataStorageInterface::FQueryDescription& Description,
 		FTypedElementExtendedQueryStore& QueryStore,
+		FTypedElementDatabaseScratchBuffer& ScratchBuffer,
 		FMassEntityQuery& NativeQuery,
 		ITypedElementDataStorageInterface::QueryCallbackRef Callback);
 
@@ -36,8 +38,12 @@ struct FTypedElementQueryProcessorData
 	FTypedElementQueryProcessorData() = default;
 	explicit FTypedElementQueryProcessorData(UMassProcessor& Owner);
 
-	bool CommonQueryConfiguration(UMassProcessor& InOwner, FTypedElementExtendedQuery& InQuery, 
-		FTypedElementExtendedQueryStore& InQueryStore, TArrayView<FMassEntityQuery> Subqueries);
+	bool CommonQueryConfiguration(
+		UMassProcessor& InOwner, 
+		FTypedElementExtendedQuery& InQuery, 
+		FTypedElementExtendedQueryStore& InQueryStore, 
+		FTypedElementDatabaseScratchBuffer& InScratchBuffer,
+		TArrayView<FMassEntityQuery> Subqueries);
 	static EMassProcessingPhase MapToMassProcessingPhase(ITypedElementDataStorageInterface::EQueryTickPhase Phase);
 	FString GetProcessorName() const;
 
@@ -65,6 +71,7 @@ struct FTypedElementQueryProcessorData
 	
 	FTypedElementExtendedQuery* ParentQuery{ nullptr };
 	FTypedElementExtendedQueryStore* QueryStore{ nullptr };
+	FTypedElementDatabaseScratchBuffer* ScratchBuffer{ nullptr };
 	FMassEntityQuery NativeQuery;
 };
 
@@ -80,10 +87,16 @@ public:
 	UTypedElementQueryProcessorCallbackAdapterProcessorBase();
 
 	FMassEntityQuery& GetQuery();
-	virtual bool ConfigureQueryCallback(FTypedElementExtendedQuery& Query, FTypedElementExtendedQueryStore& QueryStore);
+	virtual bool ConfigureQueryCallback(
+		FTypedElementExtendedQuery& Query, 
+		FTypedElementExtendedQueryStore& QueryStore,
+		FTypedElementDatabaseScratchBuffer& ScratchBuffer);
 
 protected:
-	bool ConfigureQueryCallbackData(FTypedElementExtendedQuery& Query, FTypedElementExtendedQueryStore& QueryStore, 
+	bool ConfigureQueryCallbackData(
+		FTypedElementExtendedQuery& Query, 
+		FTypedElementExtendedQueryStore& QueryStore, 
+		FTypedElementDatabaseScratchBuffer& ScratchBuffer,
 		TArrayView<FMassEntityQuery> Subqueries);
 	void ConfigureQueries() override;
 	void Execute(FMassEntityManager& EntityManager, FMassExecutionContext& TargetParentQuery) override;
@@ -116,7 +129,10 @@ class UTypedElementQueryProcessorCallbackAdapterProcessorWith1Subquery final : p
 	GENERATED_BODY()
 
 public:
-	bool ConfigureQueryCallback(FTypedElementExtendedQuery& Query, FTypedElementExtendedQueryStore& QueryStore) override;
+	bool ConfigureQueryCallback(
+		FTypedElementExtendedQuery& Query, 
+		FTypedElementExtendedQueryStore& QueryStore, 
+		FTypedElementDatabaseScratchBuffer& ScratchBuffer) override;
 
 private:
 	UPROPERTY(transient)
@@ -129,7 +145,10 @@ class UTypedElementQueryProcessorCallbackAdapterProcessorWith2Subqueries final :
 	GENERATED_BODY()
 
 public:
-	bool ConfigureQueryCallback(FTypedElementExtendedQuery& Query, FTypedElementExtendedQueryStore& QueryStore) override;
+	bool ConfigureQueryCallback(
+		FTypedElementExtendedQuery& Query, 
+		FTypedElementExtendedQueryStore& QueryStore,
+		FTypedElementDatabaseScratchBuffer& ScratchBuffer) override;
 
 private:
 	UPROPERTY(transient)
@@ -142,7 +161,10 @@ class UTypedElementQueryProcessorCallbackAdapterProcessorWith3Subqueries final :
 	GENERATED_BODY()
 
 public:
-	bool ConfigureQueryCallback(FTypedElementExtendedQuery& Query, FTypedElementExtendedQueryStore& QueryStore) override;
+	bool ConfigureQueryCallback(
+		FTypedElementExtendedQuery& Query, 
+		FTypedElementExtendedQueryStore& QueryStore,
+		FTypedElementDatabaseScratchBuffer& ScratchBuffer) override;
 
 private:
 	UPROPERTY(transient)
@@ -155,7 +177,10 @@ class UTypedElementQueryProcessorCallbackAdapterProcessorWith4Subqueries final :
 	GENERATED_BODY()
 
 public:
-	bool ConfigureQueryCallback(FTypedElementExtendedQuery& Query, FTypedElementExtendedQueryStore& QueryStore) override;
+	bool ConfigureQueryCallback(
+		FTypedElementExtendedQuery& Query, 
+		FTypedElementExtendedQueryStore& QueryStore,
+		FTypedElementDatabaseScratchBuffer& ScratchBuffer) override;
 
 private:
 	UPROPERTY(transient)
@@ -176,11 +201,12 @@ public:
 	FMassEntityQuery& GetQuery();
 	const UScriptStruct* GetObservedType() const;
 	EMassObservedOperation GetObservedOperation() const;
-	virtual bool ConfigureQueryCallback(FTypedElementExtendedQuery& Query, FTypedElementExtendedQueryStore& QueryStore);
+	virtual bool ConfigureQueryCallback(
+		FTypedElementExtendedQuery& Query, FTypedElementExtendedQueryStore& QueryStore, FTypedElementDatabaseScratchBuffer& ScratchBuffer);
 
 protected:
 	bool ConfigureQueryCallbackData(FTypedElementExtendedQuery& Query, FTypedElementExtendedQueryStore& QueryStore,
-		TArrayView<FMassEntityQuery> Subqueries);
+		FTypedElementDatabaseScratchBuffer& ScratchBuffer, TArrayView<FMassEntityQuery> Subqueries);
 	void ConfigureQueries() override;
 	void Execute(FMassEntityManager& EntityManager, FMassExecutionContext& TargetParentQuery) override;
 
@@ -213,7 +239,10 @@ class UTypedElementQueryObserverCallbackAdapterProcessorWith1Subquery final : pu
 	GENERATED_BODY()
 
 public:
-	bool ConfigureQueryCallback(FTypedElementExtendedQuery& Query, FTypedElementExtendedQueryStore& QueryStore) override;
+	bool ConfigureQueryCallback(
+		FTypedElementExtendedQuery& Query, 
+		FTypedElementExtendedQueryStore& QueryStore,
+		FTypedElementDatabaseScratchBuffer& ScratchBuffer) override;
 
 private:
 	UPROPERTY(transient)
@@ -226,7 +255,10 @@ class UTypedElementQueryObserverCallbackAdapterProcessorWith2Subqueries final :
 	GENERATED_BODY()
 
 public:
-	bool ConfigureQueryCallback(FTypedElementExtendedQuery& Query, FTypedElementExtendedQueryStore& QueryStore) override;
+	bool ConfigureQueryCallback(
+		FTypedElementExtendedQuery& Query, 
+		FTypedElementExtendedQueryStore& QueryStore,
+		FTypedElementDatabaseScratchBuffer& ScratchBuffer) override;
 
 private:
 	UPROPERTY(transient)
@@ -239,7 +271,10 @@ class UTypedElementQueryObserverCallbackAdapterProcessorWith3Subqueries final :
 	GENERATED_BODY()
 
 public:
-	bool ConfigureQueryCallback(FTypedElementExtendedQuery& Query, FTypedElementExtendedQueryStore& QueryStore) override;
+	bool ConfigureQueryCallback(
+		FTypedElementExtendedQuery& Query, 
+		FTypedElementExtendedQueryStore& QueryStore,
+		FTypedElementDatabaseScratchBuffer& ScratchBuffer) override;
 
 private:
 	UPROPERTY(transient)
@@ -252,7 +287,10 @@ class UTypedElementQueryObserverCallbackAdapterProcessorWith4Subqueries final :
 	GENERATED_BODY()
 
 public:
-	bool ConfigureQueryCallback(FTypedElementExtendedQuery& Query, FTypedElementExtendedQueryStore& QueryStore) override;
+	bool ConfigureQueryCallback(
+		FTypedElementExtendedQuery& Query, 
+		FTypedElementExtendedQueryStore& QueryStore,
+		FTypedElementDatabaseScratchBuffer& ScratchBuffer) override;
 
 private:
 	UPROPERTY(transient)
diff --git a/Engine/Plugins/Experimental/TypedElementsDataStorage/Source/TypedElementsDataStorage/Private/Queries/TypedElementExtendedQueryStore.cpp b/Engine/Plugins/Experimental/TypedElementsDataStorage/Source/TypedElementsDataStorage/Private/Queries/TypedElementExtendedQueryStore.cpp
--- a/Engine/Plugins/Experimental/TypedElementsDataStorage/Source/TypedElementsDataStorage/Private/Queries/TypedElementExtendedQueryStore.cpp
+++ b/Engine/Plugins/Experimental/TypedElementsDataStorage/Source/TypedElementsDataStorage/Private/Queries/TypedElementExtendedQueryStore.cpp
@@ -10,7 +10,10 @@
 const ITypedElementDataStorageInterface::FQueryDescription FTypedElementExtendedQueryStore::EmptyDescription{};
 
 FTypedElementExtendedQueryStore::Handle FTypedElementExtendedQueryStore::RegisterQuery(
-	ITypedElementDataStorageInterface::FQueryDescription Query, FMassEntityManager& EntityManager, FMassProcessingPhaseManager& PhaseManager)
+	ITypedElementDataStorageInterface::FQueryDescription Query, 
+	FTypedElementDatabaseScratchBuffer& ScratchBuffer, 
+	FMassEntityManager& EntityManager, 
+	FMassProcessingPhaseManager& PhaseManager)
 {
 	FTypedElementExtendedQueryStore::Handle Result = Queries.Emplace();
 	FTypedElementExtendedQuery& StoredQuery = GetMutableChecked(Result);
@@ -21,7 +24,7 @@ FTypedElementExtendedQueryStore::Handle FTypedElementExtendedQueryStore::Registe
 	bContinueSetup = bContinueSetup &&	SetupConditions(StoredQuery.Description, NativeQuery);
 	bContinueSetup = bContinueSetup &&	SetupDependencies(StoredQuery.Description, NativeQuery);
 	bContinueSetup = bContinueSetup &&	SetupTickGroupDefaults(StoredQuery.Description);
-	bContinueSetup = bContinueSetup &&	SetupProcessors(Result, StoredQuery, EntityManager, PhaseManager);
+	bContinueSetup = bContinueSetup &&	SetupProcessors(Result, StoredQuery, ScratchBuffer, EntityManager, PhaseManager);
 	
 	if (!bContinueSetup)
 	{
@@ -280,17 +283,18 @@ TypedElementDataStorage::FQueryResult FTypedElementExtendedQueryStore::RunQuery(
 
 	return Result;
 }
-
 void FTypedElementExtendedQueryStore::RunPhasePreambleQueries(FMassEntityManager& EntityManager,
-	ITypedElementDataStorageInterface::EQueryTickPhase Phase, float DeltaTime)
+	FTypedElementDatabaseScratchBuffer& ScratchBuffer, ITypedElementDataStorageInterface::EQueryTickPhase Phase, float DeltaTime)
 {
-	RunPhasePreOrPostAmbleQueries(EntityManager, Phase, DeltaTime, PhasePreparationQueries[static_cast<QueryTickPhaseType>(Phase)]);
+	RunPhasePreOrPostAmbleQueries(EntityManager, ScratchBuffer, Phase, DeltaTime, 
+		PhasePreparationQueries[static_cast<QueryTickPhaseType>(Phase)]);
 }
 
 void FTypedElementExtendedQueryStore::RunPhasePostambleQueries(FMassEntityManager& EntityManager,
-	ITypedElementDataStorageInterface::EQueryTickPhase Phase, float DeltaTime)
+	FTypedElementDatabaseScratchBuffer& ScratchBuffer, ITypedElementDataStorageInterface::EQueryTickPhase Phase, float DeltaTime)
 {
-	RunPhasePreOrPostAmbleQueries(EntityManager, Phase, DeltaTime, PhaseFinalizationQueries[static_cast<QueryTickPhaseType>(Phase)]);
+	RunPhasePreOrPostAmbleQueries(EntityManager, ScratchBuffer, Phase, DeltaTime, 
+		PhaseFinalizationQueries[static_cast<QueryTickPhaseType>(Phase)]);
 }
 
 void FTypedElementExtendedQueryStore::DebugPrintQueryCallbacks(FOutputDevice& Output) const
@@ -578,7 +582,7 @@ bool FTypedElementExtendedQueryStore::SetupTickGroupDefaults(ITypedElementDataSt
 }
 
 bool FTypedElementExtendedQueryStore::SetupProcessors(Handle Query, FTypedElementExtendedQuery& StoredQuery, 
-	FMassEntityManager& EntityManager, FMassProcessingPhaseManager& PhaseManager)
+	FTypedElementDatabaseScratchBuffer& ScratchBuffer, FMassEntityManager& EntityManager, FMassProcessingPhaseManager& PhaseManager)
 {
 	using DSI = ITypedElementDataStorageInterface;
 
@@ -599,7 +603,7 @@ bool FTypedElementExtendedQueryStore::SetupProcessors(Handle Query, FTypedElemen
 		if (StoredQuery.Processor->IsA<UTypedElementQueryProcessorCallbackAdapterProcessorBase>())
 		{
 			if (static_cast<UTypedElementQueryProcessorCallbackAdapterProcessorBase*>(StoredQuery.Processor.Get())->
-				ConfigureQueryCallback(StoredQuery, *this))
+				ConfigureQueryCallback(StoredQuery, *this, ScratchBuffer))
 			{
 				PhaseManager.RegisterDynamicProcessor(*StoredQuery.Processor);
 			}
@@ -613,7 +617,7 @@ bool FTypedElementExtendedQueryStore::SetupProcessors(Handle Query, FTypedElemen
 			if (UTypedElementQueryObserverCallbackAdapterProcessorBase* Observer =
 				static_cast<UTypedElementQueryObserverCallbackAdapterProcessorBase*>(StoredQuery.Processor.Get()))
 			{
-				Observer->ConfigureQueryCallback(StoredQuery, *this);
+				Observer->ConfigureQueryCallback(StoredQuery, *this, ScratchBuffer);
 				EntityManager.GetObserverManager().AddObserverInstance(*Observer->GetObservedType(), Observer->GetObservedOperation(), *Observer);
 			}
 			else
@@ -674,7 +678,8 @@ void FTypedElementExtendedQueryStore::UnregisterPostambleQuery(ITypedElementData
 }
 
 void FTypedElementExtendedQueryStore::RunPhasePreOrPostAmbleQueries(FMassEntityManager& EntityManager,
-	ITypedElementDataStorageInterface::EQueryTickPhase Phase, float DeltaTime, TArray<Handle>& QueryHandles)
+	FTypedElementDatabaseScratchBuffer& ScratchBuffer, ITypedElementDataStorageInterface::EQueryTickPhase Phase, 
+	float DeltaTime, TArray<Handle>& QueryHandles)
 {
 	if (!QueryHandles.IsEmpty())
 	{
@@ -682,7 +687,7 @@ void FTypedElementExtendedQueryStore::RunPhasePreOrPostAmbleQueries(FMassEntityM
 		for (Handle Query : QueryHandles)
 		{
 			FTypedElementExtendedQuery& QueryData = Queries.Get(Query);
-			Executor.ExecuteQuery(QueryData.Description, *this, QueryData.NativeQuery, QueryData.Description.Callback.Function);
+			Executor.ExecuteQuery(QueryData.Description, *this, ScratchBuffer, QueryData.NativeQuery, QueryData.Description.Callback.Function);
 		}
 	}
 }
diff --git a/Engine/Plugins/Experimental/TypedElementsDataStorage/Source/TypedElementsDataStorage/Private/Queries/TypedElementExtendedQueryStore.h b/Engine/Plugins/Experimental/TypedElementsDataStorage/Source/TypedElementsDataStorage/Private/Queries/TypedElementExtendedQueryStore.h
--- a/Engine/Plugins/Experimental/TypedElementsDataStorage/Source/TypedElementsDataStorage/Private/Queries/TypedElementExtendedQueryStore.h
+++ b/Engine/Plugins/Experimental/TypedElementsDataStorage/Source/TypedElementsDataStorage/Private/Queries/TypedElementExtendedQueryStore.h
@@ -10,8 +10,9 @@
 
 struct FMassEntityManager;
 struct FMassProcessingPhaseManager;
-class FOutputDevice;
 class UMassProcessor;
+class FOutputDevice;
+class FTypedElementDatabaseScratchBuffer;
 
 struct FTypedElementExtendedQuery
 {
@@ -37,7 +38,10 @@ public:
 	 */
 
 	/** Adds a new query to the store and initializes the query with the provided arguments. */
-	Handle RegisterQuery(ITypedElementDataStorageInterface::FQueryDescription Query, FMassEntityManager& EntityManager, 
+	Handle RegisterQuery(
+		ITypedElementDataStorageInterface::FQueryDescription Query, 
+		FTypedElementDatabaseScratchBuffer& ScratchBuffer,
+		FMassEntityManager& EntityManager,
 		FMassProcessingPhaseManager& PhaseManager);
 	/** Removes the query at the given handle if still alive and otherwise does nothing. */
 	void UnregisterQuery(Handle Query, FMassProcessingPhaseManager& PhaseManager);
@@ -87,8 +91,16 @@ public:
 		TypedElementDataStorage::SubqueryCallbackRef Callback);
 	TypedElementDataStorage::FQueryResult RunQuery(FMassEntityManager& EntityManager, Handle Query,
 		TypedElementRowHandle Row, TypedElementDataStorage::SubqueryCallbackRef Callback);
-	void RunPhasePreambleQueries(FMassEntityManager& EntityManager, ITypedElementDataStorageInterface::EQueryTickPhase Phase, float DeltaTime);
-	void RunPhasePostambleQueries(FMassEntityManager& EntityManager, ITypedElementDataStorageInterface::EQueryTickPhase Phase, float DeltaTime);
+	void RunPhasePreambleQueries(
+		FMassEntityManager& EntityManager, 
+		FTypedElementDatabaseScratchBuffer& ScratchBuffer, 
+		ITypedElementDataStorageInterface::EQueryTickPhase Phase, 
+		float DeltaTime);
+	void RunPhasePostambleQueries(
+		FMassEntityManager& EntityManager, 
+		FTypedElementDatabaseScratchBuffer& ScratchBuffer, 
+		ITypedElementDataStorageInterface::EQueryTickPhase Phase, 
+		float DeltaTime);
 
 	void DebugPrintQueryCallbacks(FOutputDevice& Output) const;
 
@@ -121,7 +133,7 @@ private:
 	bool SetupConditions(ITypedElementDataStorageInterface::FQueryDescription& Query, FMassEntityQuery& NativeQuery);
 	bool SetupDependencies(ITypedElementDataStorageInterface::FQueryDescription& Query, FMassEntityQuery& NativeQuery);
 	bool SetupTickGroupDefaults(ITypedElementDataStorageInterface::FQueryDescription& Query);
-	bool SetupProcessors(Handle Query, FTypedElementExtendedQuery& StoredQuery,
+	bool SetupProcessors(Handle Query, FTypedElementExtendedQuery& StoredQuery, FTypedElementDatabaseScratchBuffer& ScratchBuffer,
 		FMassEntityManager& EntityManager, FMassProcessingPhaseManager& PhaseManager);
 
 	EMassFragmentAccess ConvertToNativeAccessType(ITypedElementDataStorageInterface::EQueryAccessType AccessType);
@@ -130,7 +142,7 @@ private:
 	void RegisterPostambleQuery(ITypedElementDataStorageInterface::EQueryTickPhase Phase, Handle Query);
 	void UnregisterPreambleQuery(ITypedElementDataStorageInterface::EQueryTickPhase Phase, Handle Query);
 	void UnregisterPostambleQuery(ITypedElementDataStorageInterface::EQueryTickPhase Phase, Handle Query);
-	void RunPhasePreOrPostAmbleQueries(FMassEntityManager& EntityManager,
+	void RunPhasePreOrPostAmbleQueries(FMassEntityManager& EntityManager, FTypedElementDatabaseScratchBuffer& ScratchBuffer,
 		ITypedElementDataStorageInterface::EQueryTickPhase Phase, float DeltaTime, TArray<Handle>& QueryHandles);
 
 	static const ITypedElementDataStorageInterface::FQueryDescription EmptyDescription;
diff --git a/Engine/Plugins/Experimental/TypedElementsDataStorage/Source/TypedElementsDataStorage/Private/TypedElementDatabase.cpp b/Engine/Plugins/Experimental/TypedElementsDataStorage/Source/TypedElementsDataStorage/Private/TypedElementDatabase.cpp
--- a/Engine/Plugins/Experimental/TypedElementsDataStorage/Source/TypedElementsDataStorage/Private/TypedElementDatabase.cpp
+++ b/Engine/Plugins/Experimental/TypedElementsDataStorage/Source/TypedElementsDataStorage/Private/TypedElementDatabase.cpp
@@ -110,6 +110,8 @@ void UTypedElementDatabase::OnPreMassTick(float DeltaTime)
 {
 	checkf(IsAvailable(), TEXT("Typed Element Database was ticked while it's not ready."));
 	OnUpdateDelegate.Broadcast();
+	// Recycle any full scratch memory blocks from the previous frame.
+	ScratchBuffer.RecycleBlocks();
 }
 
 TSharedPtr<FMassEntityManager> UTypedElementDatabase::GetActiveMutableEditorEntityManager()
@@ -244,9 +246,10 @@ bool UTypedElementDatabase::BatchAddRow(TypedElementTableHandle Table, TConstArr
 
 void UTypedElementDatabase::RemoveRow(TypedElementRowHandle Row)
 {
-	if (ActiveEditorEntityManager)
+	FMassEntityHandle Entity = FMassEntityHandle::FromNumber(Row);
+	if (ActiveEditorEntityManager && ActiveEditorEntityManager->IsEntityValid(Entity))
 	{
-		if (ActiveEditorEntityManager->IsEntityActive(FMassEntityHandle::FromNumber(Row)))
+		if (ActiveEditorEntityManager->IsEntityBuilt(FMassEntityHandle::FromNumber(Row)))
 		{
 			ActiveEditorEntityManager->DestroyEntity(FMassEntityHandle::FromNumber(Row));
 		}
@@ -612,7 +615,7 @@ void UTypedElementDatabase::UnregisterTickGroup(FName GroupName, EQueryTickPhase
 TypedElementQueryHandle UTypedElementDatabase::RegisterQuery(FQueryDescription&& Query)
 {
 	return (ActiveEditorEntityManager && ActiveEditorPhaseManager)
-		? Queries.RegisterQuery(MoveTemp(Query), *ActiveEditorEntityManager, *ActiveEditorPhaseManager).Handle
+		? Queries.RegisterQuery(MoveTemp(Query), ScratchBuffer, *ActiveEditorEntityManager, *ActiveEditorPhaseManager).Handle
 		: TypedElementInvalidQueryHandle;
 }
 
@@ -726,7 +729,7 @@ void UTypedElementDatabase::PreparePhase(EQueryTickPhase Phase, float DeltaTime)
 {
 	if (ActiveEditorEntityManager)
 	{
-		Queries.RunPhasePreambleQueries(*ActiveEditorEntityManager, Phase, DeltaTime);
+		Queries.RunPhasePreambleQueries(*ActiveEditorEntityManager, ScratchBuffer, Phase, DeltaTime);
 	}
 }
 
@@ -734,7 +737,7 @@ void UTypedElementDatabase::FinalizePhase(EQueryTickPhase Phase, float DeltaTime
 {
 	if (ActiveEditorEntityManager)
 	{
-		Queries.RunPhasePostambleQueries(*ActiveEditorEntityManager, Phase, DeltaTime);
+		Queries.RunPhasePostambleQueries(*ActiveEditorEntityManager, ScratchBuffer, Phase, DeltaTime);
 	}
 }
 
diff --git a/Engine/Plugins/Experimental/TypedElementsDataStorage/Source/TypedElementsDataStorage/Private/TypedElementDatabase.h b/Engine/Plugins/Experimental/TypedElementsDataStorage/Source/TypedElementsDataStorage/Private/TypedElementDatabase.h
--- a/Engine/Plugins/Experimental/TypedElementsDataStorage/Source/TypedElementsDataStorage/Private/TypedElementDatabase.h
+++ b/Engine/Plugins/Experimental/TypedElementsDataStorage/Source/TypedElementsDataStorage/Private/TypedElementDatabase.h
@@ -6,6 +6,7 @@
 #include "Elements/Interfaces/TypedElementDataStorageInterface.h"
 #include "MassArchetypeTypes.h"
 #include "Queries/TypedElementExtendedQueryStore.h"
+#include "TypedElementDatabaseScratchBuffer.h"
 #include "UObject/ObjectMacros.h"
 #include "UObject/StrongObjectPtr.h"
 
@@ -99,6 +100,7 @@ private:
 	TMap<FName, TypedElementTableHandle> TableNameLookup;
 
 	FTypedElementExtendedQueryStore Queries;
+	FTypedElementDatabaseScratchBuffer ScratchBuffer;
 
 	FTypedElementOnDataStorageUpdate OnUpdateDelegate;
 
diff --git a/Engine/Plugins/Experimental/TypedElementsDataStorage/Source/TypedElementsDataStorage/Private/TypedElementDatabaseScratchBuffer.cpp b/Engine/Plugins/Experimental/TypedElementsDataStorage/Source/TypedElementsDataStorage/Private/TypedElementDatabaseScratchBuffer.cpp
new file mode 100644
--- /dev/null
+++ b/Engine/Plugins/Experimental/TypedElementsDataStorage/Source/TypedElementsDataStorage/Private/TypedElementDatabaseScratchBuffer.cpp
@@ -0,0 +1,197 @@
+// Copyright Epic Games, Inc. All Rights Reserved.
+
+#include "TypedElementDatabaseScratchBuffer.h"
+
+FTypedElementDatabaseScratchBuffer::~FTypedElementDatabaseScratchBuffer()
+{
+	// This assumes it's being called after all threads using the scratch buffer have been shutdown already so the
+	// scratch buffers they hold have been released back into the pool.
+
+	RecycleBlocks();
+	FBlock* FrontFull = FullBlocks;
+	while (FrontFull)
+	{
+		FBlock* Next = FrontFull->NextBlock;
+		delete FrontFull;
+		FrontFull = Next;
+	}
+	FullBlocks = nullptr;
+
+	FBlock* FrontAvailable = AvailableBlocks;
+	while (FrontAvailable)
+	{
+		FBlock* Next = FrontAvailable->NextBlock;
+		delete FrontAvailable;
+		FrontAvailable = Next;
+	}
+	AvailableBlocks = nullptr;
+}
+
+void* FTypedElementDatabaseScratchBuffer::Allocate(size_t Size, size_t Alignment)
+{
+	TScopedWriterDetector Detector(AccessDetector);
+	return GetThreadLocalBlockController().Allocate(Size, Alignment);
+}
+
+void FTypedElementDatabaseScratchBuffer::RecycleBlocks()
+{
+	if (FullBlocks)
+	{
+		TScopedWriterDetector Detector(AccessDetector);
+		
+		// Destroy any objects in blocks. It walks the allocations in a block backwards and invokes the destructors on the allocated 
+		// objects. It moves to the next block once it reaches the beginning of the block. At the end the cleared blocks are reinserted 
+		// into the available blocks for reuse. This approach allows for simple lock-free appending of the destructor chain while remaining
+		// relatively simple to destroy. The dirty blocks are added in front of the available blocks instead of the end to avoid having
+		// to iterate over the entire list of available blocks.
+
+		FBlock* DirtyBlock = FullBlocks;
+		FBlock* LastDirtyBlock = FullBlocks;
+		do
+		{
+			FDestructorTail* Destructor = DirtyBlock->DestructionTail;
+			while (Destructor)
+			{
+				Destructor->Destructor(reinterpret_cast<char*>(Destructor) - Destructor->StructOffset, Destructor->InstanceCount);
+				Destructor = Destructor->PreviousTail;
+			}
+			DirtyBlock->DestructionTail = nullptr;
+
+			LastDirtyBlock = DirtyBlock;
+			DirtyBlock = DirtyBlock->NextBlock;
+		} while (DirtyBlock);
+		
+		// Attach the full blocks in front of the available blocks.
+		LastDirtyBlock->NextBlock.store(AvailableBlocks);
+		AvailableBlocks.store(FullBlocks);
+		FullBlocks = nullptr;
+	}
+}
+
+FTypedElementDatabaseScratchBuffer::FBlockController& FTypedElementDatabaseScratchBuffer::GetThreadLocalBlockController()
+{
+	thread_local static FBlockController LocalBlock = FBlockController(*this);
+	return LocalBlock;
+}
+
+void FTypedElementDatabaseScratchBuffer::ConfigureDestructorTail(
+	FDestructorTail& Destructor, DestructorFunction Callback, void* Object, int32 Count)
+{
+	GetThreadLocalBlockController().ConfigureDestructorTail(Destructor, Callback, Object, Count);
+}
+
+
+//
+// FBlockController
+//
+
+FTypedElementDatabaseScratchBuffer::FBlockController::FBlockController(FTypedElementDatabaseScratchBuffer& InOwner)
+	: Owner(InOwner)
+	, Id(InOwner.BlockControllerId++)
+{
+	Block = GetEmptyBlock();
+}
+
+FTypedElementDatabaseScratchBuffer::FBlockController::~FBlockController()
+{
+	RecycleBlock();
+}
+
+void* FTypedElementDatabaseScratchBuffer::FBlockController::Allocate(size_t Size, size_t Alignment)
+{
+	checkf(Alignment <= alignof(FBlock), TEXT("Alignment of %i for allocation in database scratch buffer exceeds maximal alignment of %i."),
+		static_cast<int>(Alignment), static_cast<int>(alignof(FBlock)));
+
+	if (Size <= MaxAllocationSize())
+	{
+		size_t Index = Align(Block->Front, Alignment);
+		size_t NewFront = Index + Size;
+
+		if (NewFront > FBlock::BlockSize)
+		{
+			// No more space left in the block so recycle the current one and create a new one.
+			RecycleBlock();
+			Block = GetEmptyBlock();
+			Index = 0;
+			NewFront = Size;
+		}
+
+		Block->Front = NewFront;
+		return Block->Buffer + Index;
+	}
+	else
+	{
+		struct FExtendedBufferStore
+		{
+			void* ExtendedBuffer;
+			FDestructorTail Destructor;
+		};
+
+		// Allocate memory directly as the requested size doesn't fit in a block.
+		void* ExtendedBuffer = FMemory::Malloc(Size, Alignment);
+		
+		// Add an entry to the block with the sole purpose of deleting the extended buffer.
+		void* BufferStoreAddress = Allocate(sizeof(FExtendedBufferStore), alignof(FExtendedBufferStore));
+		FExtendedBufferStore* BufferStore = reinterpret_cast<FExtendedBufferStore*>(BufferStoreAddress);
+		BufferStore->ExtendedBuffer = ExtendedBuffer;
+
+		ConfigureDestructorTail(BufferStore->Destructor,
+			[](void* Object, [[maybe_unused]] int32 ObjectCount)
+			{
+				FMemory::Free(reinterpret_cast<FExtendedBufferStore*>(Object)->ExtendedBuffer);
+			}, BufferStore);
+
+		return ExtendedBuffer;
+	}
+}
+
+FTypedElementDatabaseScratchBuffer::FBlock* FTypedElementDatabaseScratchBuffer::FBlockController::GetEmptyBlock()
+{
+	FBlock* FrontBlock = Owner.AvailableBlocks;
+	while (FrontBlock != nullptr)
+	{
+		// If this is not zero it means another thread has already claimed this block so try again with the next block.
+		uint32 CurrentOwner = 0;
+		if (FrontBlock->Owner.compare_exchange_strong(CurrentOwner, Id))
+		{
+			Owner.AvailableBlocks.store(FrontBlock->NextBlock.load());
+			FrontBlock->NextBlock = nullptr;
+			return FrontBlock;
+		}
+		else
+		{
+			// Get the new front block. It can happen due to the thread timing that the same block is retrieved, but 
+			// avoiding this adds a lot more complexity and contention on the front block is low.
+			FrontBlock = Owner.AvailableBlocks;
+		}
+
+	};
+	
+	// If null then there are no more buffers left so create a new one.
+	FBlock* NewBlock = new FBlock();
+	NewBlock->Owner = Id;
+	return NewBlock;
+}
+
+void FTypedElementDatabaseScratchBuffer::FBlockController::RecycleBlock()
+{
+	Block->Front = 0;
+	Block->Owner = 0;
+
+	// Replace the front of the linked list with full blocks with the latest full block.
+	FBlock* FrontBlock = Owner.FullBlocks;
+	do
+	{
+		Block->NextBlock = FrontBlock;
+	} while (!Owner.FullBlocks.compare_exchange_strong(FrontBlock, Block));
+}
+
+void FTypedElementDatabaseScratchBuffer::FBlockController::ConfigureDestructorTail(
+	FDestructorTail& Destructor, DestructorFunction Callback, void* Object, int32 Count)
+{
+	Destructor.Destructor = Callback;
+	Destructor.PreviousTail = Block->DestructionTail;
+	Destructor.StructOffset = static_cast<int32>(reinterpret_cast<char*>(&Destructor) - reinterpret_cast<char*>(Object));
+	Destructor.InstanceCount = Count;
+	Block->DestructionTail = &Destructor;
+}
diff --git a/Engine/Plugins/Experimental/TypedElementsDataStorage/Source/TypedElementsDataStorage/Private/TypedElementDatabaseScratchBuffer.h b/Engine/Plugins/Experimental/TypedElementsDataStorage/Source/TypedElementsDataStorage/Private/TypedElementDatabaseScratchBuffer.h
new file mode 100644
--- /dev/null
+++ b/Engine/Plugins/Experimental/TypedElementsDataStorage/Source/TypedElementsDataStorage/Private/TypedElementDatabaseScratchBuffer.h
@@ -0,0 +1,164 @@
+// Copyright Epic Games, Inc. All Rights Reserved.
+
+#pragma once
+
+#include <atomic>
+#include <type_traits>
+#include "Misc/MTAccessDetector.h"
+#include "Templates/AlignmentTemplates.h"
+#include "Templates/UnrealTemplate.h"
+
+// A thread-safe memory allocator that uses linear allocation using a chain of recycled memory blocks.
+// This provides a fast and lightweight way to allocate temporary memory for intermediate values that will
+// live at best until the end of the frame. Objects that require destruction will be deleted once a block is
+// recycled. A block is only recycled when it's full, so during low activity on a particular thread it can
+// take significantly longer than a single frame before an object is destroyed.
+class FTypedElementDatabaseScratchBuffer
+{
+	friend class FBlockController;
+public:
+	~FTypedElementDatabaseScratchBuffer();
+
+	void* Allocate(size_t Size, size_t Alignment);
+	template<typename T, typename... ArgTypes>
+	T* Emplace(ArgTypes... Args);
+	template<typename T, typename... ArgTypes>
+	T* EmplaceArray(int32 Count, const ArgTypes&... Args);
+
+	// Makes previously filled up blocks available again. This can only safely be called when Allocate(...) can not
+	// be called. As allocations should only be called when Mass is running processors, it means this function
+	// can be safely called when the processors are not running, i.e. at the end of a tick.
+	void RecycleBlocks();
+
+	constexpr static int32 MaxAllocationSize();
+
+private:
+	using DestructorFunction = void (*)(void* Object, int32 ObjectCount);
+
+	// Track objects that need to be deleted. This will happen when a chain of blocks is recycled. This can
+	// mean that an objects stays around for several frames if a block isn't heavily used.
+	struct FDestructorTail
+	{
+		FDestructorTail* PreviousTail;
+		DestructorFunction Destructor;
+		int32 StructOffset;
+		int32 InstanceCount;
+	};
+
+	// Use the largest SIMD alignment for current hardware just in case. The buffer is sufficiently large
+	// that the overhead won't matter.
+	struct alignas(64) FBlock
+	{
+		constexpr static size_t BlockSize = 64 * 1024; // 64kb blocks
+		
+		char Buffer[BlockSize]; // Keep at top for alignment.
+		FDestructorTail* DestructionTail = nullptr;
+		std::atomic<FBlock*> NextBlock = nullptr;
+		std::atomic<uint32> Owner = 0;
+		// Offset into the buffer where the next allocation starts.
+		uint32 Front = 0;
+	};
+
+	// Simple wrapper around the buffer. This is to guarantee that when the thread gets destroyed the buffer is cleaned up.
+	struct FBlockController
+	{
+		explicit FBlockController(FTypedElementDatabaseScratchBuffer& InOwner);
+		~FBlockController();
+
+		void* Allocate(size_t Size, size_t Alignment);
+		FBlock* GetEmptyBlock();
+		void RecycleBlock();
+		void ConfigureDestructorTail(FDestructorTail& Destructor, DestructorFunction Callback, void* Object, int32 Count = 1);
+
+		FTypedElementDatabaseScratchBuffer& Owner;
+		FBlock* Block;
+		uint32 Id;
+	};
+
+	// Returns the block controller that's unique to the thread that calls this function.
+	FBlockController& GetThreadLocalBlockController();
+	void ConfigureDestructorTail(FDestructorTail& Destructor, DestructorFunction Callback, void* Object, int32 Count = 1);
+
+	// Detect any invalid calls to recycle blocks while allocating memory.
+	FRWRecursiveAccessDetector AccessDetector;
+
+	std::atomic<FBlock*> AvailableBlocks = nullptr;
+	std::atomic<FBlock*> FullBlocks = nullptr;
+	// Running counter so each FBlockController gets a unique id to identify the blocks assigned to them. Zero is reserved 
+	// to indicate it's not in use by a block controller.
+	std::atomic<uint32> BlockControllerId = 1;
+};
+
+
+//
+// Implementations
+//
+
+template<typename T, typename... ArgTypes>
+T* FTypedElementDatabaseScratchBuffer::Emplace(ArgTypes... Args)
+{
+	if constexpr (std::is_trivially_destructible_v<T>)
+	{
+		return new(Allocate(sizeof(T), alignof(T))) T(Forward<ArgTypes>(Args)...);
+	}
+	else
+	{
+		struct Wrapper
+		{
+			T Instance;
+			FDestructorTail Destructor;
+		};
+		Wrapper* Result = reinterpret_cast<Wrapper*>(Allocate(sizeof(Wrapper), alignof(Wrapper)));
+		new(&Result->Instance) T(Forward<ArgTypes>(Args)...);
+
+		ConfigureDestructorTail(Result->Destructor, 
+			[](void* Object, [[maybe_unused]] int32 ObjectCount)
+			{ 
+				reinterpret_cast<T*>(Object)->~T(); 
+			}, &Result->Instance);
+
+		return &Result->Instance;
+	}
+}
+
+template<typename T, typename... ArgTypes>
+T* FTypedElementDatabaseScratchBuffer::EmplaceArray(int32 Count, const ArgTypes&... Args)
+{
+	T* Result;	
+	if constexpr (std::is_trivially_destructible_v<T>)
+	{
+		Result = reinterpret_cast<T*>(Allocate(sizeof(T) * Count, alignof(T)));
+	}
+	else
+	{
+		constexpr size_t Alignment = std::max(alignof(T), alignof(FDestructorTail));
+		int32 DestructorTailOffset = Align(sizeof(T) * Count, alignof(FDestructorTail));
+		int32 Size = DestructorTailOffset + sizeof(FDestructorTail);
+		Result = reinterpret_cast<T*>(Allocate(Size, Alignment));
+
+		FDestructorTail* Destructor = reinterpret_cast<FDestructorTail*>(reinterpret_cast<char*>(Result) + DestructorTailOffset);
+		ConfigureDestructorTail(*Destructor,
+			[](void* Object, int32 ObjectCount)
+			{
+				T* Array = reinterpret_cast<T*>(Object);
+				for (int32 Count = 0; Count < ObjectCount; ++Count)
+				{
+					Array->~T();
+					++Array;
+				}
+			}, Result, Count);
+	}
+
+	for (int32 Index = 0; Index < Count; ++Index)
+	{
+		new(Result + Index) T(Args...);
+	}
+
+	return Result;
+}
+
+constexpr int32 FTypedElementDatabaseScratchBuffer::MaxAllocationSize()
+{
+	// Leave room for at least one destructor tail in case a single large object is added that requires destruction.
+	return FBlock::BlockSize - sizeof(FDestructorTail);
+}
diff --git a/Engine/Source/Runtime/TypedElementFramework/Public/Elements/Interfaces/TypedElementQueryStorageInterfaces.h b/Engine/Source/Runtime/TypedElementFramework/Public/Elements/Interfaces/TypedElementQueryStorageInterfaces.h
--- a/Engine/Source/Runtime/TypedElementFramework/Public/Elements/Interfaces/TypedElementQueryStorageInterfaces.h
+++ b/Engine/Source/Runtime/TypedElementFramework/Public/Elements/Interfaces/TypedElementQueryStorageInterfaces.h
@@ -6,6 +6,7 @@
 #include "Elements/Common/TypedElementHandles.h"
 #include "Elements/Common/TypedElementQueryTypes.h"
 #include "UObject/WeakObjectPtrTemplates.h"
+#include "Templates/Function.h"
 
 class UClass;
 class UObject;
@@ -82,6 +83,8 @@ namespace TypedElementDataStorage
 	 */
 	struct IQueryContext : public ICommonQueryContext
 	{
+		using ObjectMoveOperator = void (*)(void* Destination, void* Source);
+
 		virtual ~IQueryContext() = default;
 
 		/** Returns an immutable instance of the requested dependency or a nullptr if not found. */
@@ -105,7 +108,19 @@ namespace TypedElementDataStorage
 		 * group.
 		 */
 		virtual void RemoveRows(TConstArrayView<RowHandle> Rows) = 0;
-
+		/**
+		 * Add a new unitialized column of the provided object type. The addition will not be immediately done but delayed until the
+		 * end of the tick group. The returned address is to a temporary object that's held until the column has been setup. This can be
+		 * used to set data while in the query callback but shouldn't be used outside the callback as the object will be removed.
+		 */
+		virtual void* AddColumnUnitialized(RowHandle Row, const UScriptStruct* ObjectType) = 0;
+		/**
+		 * Add a new unitialized column of the provided object type. The addition will not be immediately done but delayed until the
+		 * end of the tick group. The returned address is to a temporary object that's held until the column has been setup. This can be
+		 * used to set data while in the query callback but shouldn't be used outside the callback as the object will be removed. The
+		 * temporary object will be moved to it's final location using the provided move operator callback instead of being copied.
+		 */
+		virtual void* AddColumnUnitialized(RowHandle Row, const UScriptStruct* ObjectType, ObjectMoveOperator Mover) = 0;
 		/**
 		 * Adds new empty columns to a row of the provided type. The addition will not be immediately done but delayed until the end of the
 		 * tick group.
@@ -116,6 +131,7 @@ namespace TypedElementDataStorage
 		 * tick group.
 		 */
 		virtual void AddColumns(TConstArrayView<RowHandle> Rows, TConstArrayView<const UScriptStruct*> ColumnTypes) = 0;
+
 		/**
 		 * Removes columns of the provided types from a row. The removal will not be immediately done but delayed until the end of the
 		 * tick group.
@@ -155,6 +171,8 @@ namespace TypedElementDataStorage
 
 		// Utility functions
 
+		template<typename ColumnType>
+		ColumnType& AddColumn(RowHandle Row, ColumnType&& Column);
 		template<typename... Columns>
 		void AddColumns(RowHandle Row);
 		template<typename... Columns>
@@ -187,6 +205,28 @@ namespace TypedElementDataStorage
 		return reinterpret_cast<Column*>(GetMutableColumn(Column::StaticStruct()));
 	}
 
+	template<typename ColumnType>
+	ColumnType& IQueryContext::AddColumn(RowHandle Row, ColumnType&& Column)
+	{
+		UScriptStruct* TypeInfo = ColumnType::StaticStruct();
+
+		if constexpr (std::is_move_constructible_v<ColumnType>)
+		{
+			void* Address = AddColumnUnitialized(Row, TypeInfo,
+				[](void* Destination, void* Source)
+				{
+					new(Destination) ColumnType(MoveTemp(*reinterpret_cast<ColumnType*>(Source)));
+				});
+			return *(new(Address) ColumnType(Forward<ColumnType>(Column)));
+		}
+		else
+		{
+			void* Address = AddColumnUnitialized(Row, TypeInfo);
+			TypeInfo->CopyScriptStruct(Address, &Column);
+			return *reinterpret_cast<ColumnType*>(Address);
+		}
+	}
+
 	template<typename... Columns>
 	void IQueryContext::AddColumns(RowHandle Row)
 	{
