Fixed Mac compiler error.

diff --git a/Engine/Plugins/Experimental/TypedElementsDataStorage/Source/TypedElementsDataStorage/Private/Processors/TypedElementProcessorAdaptors.cpp b/Engine/Plugins/Experimental/TypedElementsDataStorage/Source/TypedElementsDataStorage/Private/Processors/TypedElementProcessorAdaptors.cpp
--- a/Engine/Plugins/Experimental/TypedElementsDataStorage/Source/TypedElementsDataStorage/Private/Processors/TypedElementProcessorAdaptors.cpp
+++ b/Engine/Plugins/Experimental/TypedElementsDataStorage/Source/TypedElementsDataStorage/Private/Processors/TypedElementProcessorAdaptors.cpp
@@ -245,7 +245,7 @@ public:
 		struct FAddValueColumnWithDestructor : FAddValueColumn
 		{
 			FAddValueColumnWithDestructor(const UScriptStruct* InFragmentType, FMassEntityHandle InEntity, void* InObject)
-				: FAddValueColumn(InFragmentType, InEntity, InObject)
+				: FAddValueColumn({ .FragmentType = InFragmentType, .Entity = InEntity, .Object = InObject })
 			{}
 
 			~FAddValueColumnWithDestructor()
@@ -289,7 +289,7 @@ public:
 		{
 			FAddMoveableValueColumnWithDestructor(
 				ObjectMoveOperator InMover, const UScriptStruct* InFragmentType, FMassEntityHandle InEntity, void* InObject)
-				: FAddMoveableValueColumn(InMover, InFragmentType, InEntity, InObject)
+				: FAddMoveableValueColumn({ .Mover = InMover, .FragmentType = InFragmentType, .Entity = InEntity, .Object = InObject })
 			{}
 
 			~FAddMoveableValueColumnWithDestructor()
