Another Mac compiler error fix.

#nojira

diff --git a/Engine/Plugins/Experimental/TypedElementsDataStorage/Source/TypedElementsDataStorage/Private/Processors/TypedElementProcessorAdaptors.cpp b/Engine/Plugins/Experimental/TypedElementsDataStorage/Source/TypedElementsDataStorage/Private/Processors/TypedElementProcessorAdaptors.cpp
--- a/Engine/Plugins/Experimental/TypedElementsDataStorage/Source/TypedElementsDataStorage/Private/Processors/TypedElementProcessorAdaptors.cpp
+++ b/Engine/Plugins/Experimental/TypedElementsDataStorage/Source/TypedElementsDataStorage/Private/Processors/TypedElementProcessorAdaptors.cpp
@@ -240,12 +240,19 @@ public:
 			const UScriptStruct* FragmentType;
 			FMassEntityHandle Entity;
 			void* Object;
+
+			FAddValueColumn() = default;
+			FAddValueColumn(const UScriptStruct* InFragmentType, FMassEntityHandle InEntity, void* InObject)
+				: FragmentType(InFragmentType)
+				, Entity(InEntity)
+				, Object(InObject)
+			{}
 		};
 
 		struct FAddValueColumnWithDestructor : FAddValueColumn
 		{
 			FAddValueColumnWithDestructor(const UScriptStruct* InFragmentType, FMassEntityHandle InEntity, void* InObject)
-				: FAddValueColumn({ .FragmentType = InFragmentType, .Entity = InEntity, .Object = InObject })
+				: FAddValueColumn(InFragmentType, InEntity, InObject)
 			{}
 
 			~FAddValueColumnWithDestructor()
@@ -283,13 +290,21 @@ public:
 			const UScriptStruct* FragmentType;
 			FMassEntityHandle Entity;
 			void* Object;
+
+			FAddMoveableValueColumn() = default;
+			FAddMoveableValueColumn(ObjectMoveOperator InMover, const UScriptStruct* InFragmentType, FMassEntityHandle InEntity, void* InObject)
+				: Mover(InMover)
+				, FragmentType(InFragmentType)
+				, Entity(InEntity)
+				, Object(InObject)
+			{}
 		};
 
 		struct FAddMoveableValueColumnWithDestructor : FAddMoveableValueColumn
 		{
 			FAddMoveableValueColumnWithDestructor(
 				ObjectMoveOperator InMover, const UScriptStruct* InFragmentType, FMassEntityHandle InEntity, void* InObject)
-				: FAddMoveableValueColumn({ .Mover = InMover, .FragmentType = InFragmentType, .Entity = InEntity, .Object = InObject })
+				: FAddMoveableValueColumn(InMover, InFragmentType, InEntity, InObject)
 			{}
 
 			~FAddMoveableValueColumnWithDestructor()
