Fixed NoPCH compile error on Mac.

#nojira

diff --git a/Engine/Plugins/Experimental/TypedElementsDataStorage/Source/TypedElementsDataStorage/Private/TypedElementDatabaseScratchBuffer.h b/Engine/Plugins/Experimental/TypedElementsDataStorage/Source/TypedElementsDataStorage/Private/TypedElementDatabaseScratchBuffer.h
--- a/Engine/Plugins/Experimental/TypedElementsDataStorage/Source/TypedElementsDataStorage/Private/TypedElementDatabaseScratchBuffer.h
+++ b/Engine/Plugins/Experimental/TypedElementsDataStorage/Source/TypedElementsDataStorage/Private/TypedElementDatabaseScratchBuffer.h
@@ -4,6 +4,7 @@
 
 #include <atomic>
 #include <type_traits>
+#include "GenericPlatform/GenericPlatformMath.h"
 #include "Misc/MTAccessDetector.h"
 #include "Templates/AlignmentTemplates.h"
 #include "Templates/UnrealTemplate.h"
@@ -131,7 +132,7 @@ T* FTypedElementDatabaseScratchBuffer::EmplaceArray(int32 Count, const ArgTypes&
 	}
 	else
 	{
-		constexpr size_t Alignment = std::max(alignof(T), alignof(FDestructorTail));
+		constexpr size_t Alignment = FMath::Max(alignof(T), alignof(FDestructorTail));
 		int32 DestructorTailOffset = Align(sizeof(T) * Count, alignof(FDestructorTail));
 		int32 Size = DestructorTailOffset + sizeof(FDestructorTail);
 		Result = reinterpret_cast<T*>(Allocate(Size, Alignment));
