Extended TEDS' Compatibility layer with the ability to look up objects that are referenced but aren't automatically mapped to a row.

A use case for this change is to look up a the row a component is stored on. The owning actor is mapped to a row while the components on that actor are mapped to columns. This allows looking up the correct by getting the owning actor from the component and then look that up instead.

diff --git a/Engine/Plugins/Experimental/TypedElementsDataStorage/Source/TypedElementsDataStorage/Private/TypedElementDatabaseCompatibility.cpp b/Engine/Plugins/Experimental/TypedElementsDataStorage/Source/TypedElementsDataStorage/Private/TypedElementDatabaseCompatibility.cpp
--- a/Engine/Plugins/Experimental/TypedElementsDataStorage/Source/TypedElementsDataStorage/Private/TypedElementDatabaseCompatibility.cpp
+++ b/Engine/Plugins/Experimental/TypedElementsDataStorage/Source/TypedElementsDataStorage/Private/TypedElementDatabaseCompatibility.cpp
@@ -35,6 +35,11 @@ void UTypedElementDatabaseCompatibility::Deinitialize()
 	Reset();
 }
 
+void UTypedElementDatabaseCompatibility::RegisterDealiaserCallback(ObjectToRowDealiaser Dealiaser)
+{
+	ObjectToRowDialiasers.Add(MoveTemp(Dealiaser));
+}
+
 TypedElementRowHandle UTypedElementDatabaseCompatibility::AddCompatibleObjectExplicit(UObject* Object)
 {
 	return AddCompatibleObjectExplicit(Object, StandardUObjectTable);
@@ -166,13 +171,15 @@ TypedElementRowHandle UTypedElementDatabaseCompatibility::FindRowWithCompatibleO
 		if (Actor && ActorSubsystem)
 		{
 			FMassEntityHandle Entity = ActorSubsystem->GetEntityHandleFromActor(Actor);
-			return Entity.IsValid() ? Entity.AsNumber() : TypedElementInvalidRowHandle;
+			return Entity.IsValid() ? Entity.AsNumber() : DealiasObject(Object);
 		}
 		else
 		{
 			const TypedElementRowHandle* Row = ReverseObjectLookup.Find(Object);
-			return Row ? *Row : TypedElementInvalidRowHandle;
+			return Row ? *Row : DealiasObject(Object);
 		}
+
+
 	}
 	return TypedElementInvalidRowHandle;
 }
@@ -182,7 +189,7 @@ TypedElementRowHandle UTypedElementDatabaseCompatibility::FindRowWithCompatibleO
 	if (Storage && ActorSubsystem && Storage->IsAvailable())
 	{
 		FMassEntityHandle Entity = ActorSubsystem->GetEntityHandleFromActor(Actor);
-		return Entity.IsValid() ? Entity.AsNumber() : TypedElementInvalidRowHandle;
+		return Entity.IsValid() ? Entity.AsNumber() : DealiasObject(Actor);
 	}
 	else
 	{
@@ -239,6 +246,18 @@ void UTypedElementDatabaseCompatibility::CreateStandardArchetypes()
 		FName("Editor_StandardExternalObjectTable"));
 }
 
+TypedElementRowHandle UTypedElementDatabaseCompatibility::DealiasObject(const UObject* Object) const
+{
+	for (const ObjectToRowDealiaser& Dealiaser : ObjectToRowDialiasers)
+	{
+		if (TypedElementRowHandle Row = Dealiaser(*this, Object); Row != TypedElementDataStorage::InvalidRowHandle)
+		{
+			return Row;
+		}
+	}
+	return TypedElementDataStorage::InvalidRowHandle;
+}
+
 void UTypedElementDatabaseCompatibility::Tick()
 {
 	TEDS_EVENT_SCOPE(TEXT("Compatibility Tick"))
diff --git a/Engine/Plugins/Experimental/TypedElementsDataStorage/Source/TypedElementsDataStorage/Private/TypedElementDatabaseCompatibility.h b/Engine/Plugins/Experimental/TypedElementsDataStorage/Source/TypedElementsDataStorage/Private/TypedElementDatabaseCompatibility.h
--- a/Engine/Plugins/Experimental/TypedElementsDataStorage/Source/TypedElementsDataStorage/Private/TypedElementDatabaseCompatibility.h
+++ b/Engine/Plugins/Experimental/TypedElementsDataStorage/Source/TypedElementsDataStorage/Private/TypedElementDatabaseCompatibility.h
@@ -26,6 +26,8 @@ public:
 	void Initialize(ITypedElementDataStorageInterface* StorageInterface);
 	void Deinitialize();
 
+	void RegisterDealiaserCallback(ObjectToRowDealiaser Dealiaser) override;
+
 	TypedElementRowHandle AddCompatibleObjectExplicit(UObject* Object) override;
 	TypedElementRowHandle AddCompatibleObjectExplicit(UObject* Object, TypedElementTableHandle Table) override;
 	TypedElementRowHandle AddCompatibleObjectExplicit(AActor* Actor) override;
@@ -46,6 +48,8 @@ private:
 	void Reset();
 	void CreateStandardArchetypes();
 	
+	TypedElementRowHandle DealiasObject(const UObject* Object) const;
+
 	void Tick();
 	void TickPendingActorRegistration(UWorld* EditorWorld);
 	void TickPendingUObjectRegistration();
@@ -83,6 +87,8 @@ private:
 	TMap<TypedElementTableHandle, PendingRegistration<TWeakObjectPtr<UObject>>> UObjectsPendingRegistration;
 	TMap<TypedElementTableHandle, PendingRegistration<ExternalObjectRegistration>> ExternalObjectsPendingRegistration;
 	
+	TArray<ObjectToRowDealiaser> ObjectToRowDialiasers;
+
 	TypedElementTableHandle StandardActorTable{ TypedElementInvalidTableHandle };
 	TypedElementTableHandle StandardActorWithTransformTable{ TypedElementInvalidTableHandle };
 	TypedElementTableHandle StandardUObjectTable{ TypedElementInvalidTableHandle };
diff --git a/Engine/Plugins/Experimental/TypedElementsDataStorage/Source/TypedElementsDataStorage/Private/TypedElementsDataStorage.cpp b/Engine/Plugins/Experimental/TypedElementsDataStorage/Source/TypedElementsDataStorage/Private/TypedElementsDataStorage.cpp
--- a/Engine/Plugins/Experimental/TypedElementsDataStorage/Source/TypedElementsDataStorage/Private/TypedElementsDataStorage.cpp
+++ b/Engine/Plugins/Experimental/TypedElementsDataStorage/Source/TypedElementsDataStorage/Private/TypedElementsDataStorage.cpp
@@ -97,6 +97,7 @@ void FTypedElementsDataStorageModule::StartupModule()
 				{
 					Factory->RegisterTables(*Database);
 					Factory->RegisterTickGroups(*Database);
+					Factory->RegisterDealiaser(*DatabaseCompatibility);
 					Factory->RegisterWidgetPurposes(*DatabaseUi);
 				}
 				// Second pass to call all registration that would benefit or need the registration in the previous pass.
diff --git a/Engine/Source/Runtime/TypedElementFramework/Public/Elements/Interfaces/TypedElementDataStorageCompatibilityInterface.h b/Engine/Source/Runtime/TypedElementFramework/Public/Elements/Interfaces/TypedElementDataStorageCompatibilityInterface.h
--- a/Engine/Source/Runtime/TypedElementFramework/Public/Elements/Interfaces/TypedElementDataStorageCompatibilityInterface.h
+++ b/Engine/Source/Runtime/TypedElementFramework/Public/Elements/Interfaces/TypedElementDataStorageCompatibilityInterface.h
@@ -3,6 +3,7 @@
 #pragma once
 
 #include "Elements/Interfaces/TypedElementDataStorageInterface.h"
+#include "Templates/Function.h"
 #include "UObject/Interface.h"
 #include "UObject/ObjectKey.h"
 #include "UObject/ObjectMacros.h"
@@ -29,6 +30,8 @@ class ITypedElementDataStorageCompatibilityInterface
 	GENERATED_BODY()
 
 public:
+	using ObjectToRowDealiaser = TFunction<TypedElementRowHandle(const ITypedElementDataStorageCompatibilityInterface&, const UObject*)>;
+
 	/**
 	 * @section Type-agnostic functions
 	 * These allow compatibility with any type. These do eventually fall back to the explicit versions.
@@ -53,6 +56,14 @@ public:
 	template<typename ObjectType>
 	TypedElementRowHandle FindRowWithCompatibleObject(ObjectType&& Object) const;
 
+	/**
+	 * @section Alias functions
+	 * Notifications and request can be made to the compatibility layer for objects that are stored but don't directly map to a row.
+	 * An example is a UObject represented by a column. If the UObject gets updated there's no direct mapping to the row the column is
+	 * stored in but the row still needs to be updated. For cases like this it's possible to store information to find the row that's
+	 * being aliased.
+	 */
+	virtual void RegisterDealiaserCallback(ObjectToRowDealiaser Dealiaser) = 0;
 
 	/**
 	 * @section Explicit functions
@@ -84,7 +95,7 @@ public:
 	 * faster if it's already known that the target is an actor.
 	 */
 	virtual TypedElementRowHandle FindRowWithCompatibleObjectExplicit(const AActor* Actor) const = 0;
-	/** Finds a previously stored FStructe . If not found an invalid row handle will be returned. */
+	/** Finds a previously stored FStruct. If not found an invalid row handle will be returned. */
 	virtual TypedElementRowHandle FindRowWithCompatibleObjectExplicit(const void* Object) const = 0;
 };
 
diff --git a/Engine/Source/Runtime/TypedElementFramework/Public/Elements/Interfaces/TypedElementDataStorageFactory.h b/Engine/Source/Runtime/TypedElementFramework/Public/Elements/Interfaces/TypedElementDataStorageFactory.h
--- a/Engine/Source/Runtime/TypedElementFramework/Public/Elements/Interfaces/TypedElementDataStorageFactory.h
+++ b/Engine/Source/Runtime/TypedElementFramework/Public/Elements/Interfaces/TypedElementDataStorageFactory.h
@@ -4,6 +4,7 @@
 
 #include "TypedElementDataStorageFactory.generated.h"
 
+class ITypedElementDataStorageCompatibilityInterface;
 class ITypedElementDataStorageInterface;
 class ITypedElementDataStorageUiInterface;
 
@@ -28,6 +29,8 @@ public:
 	virtual void RegisterTables(ITypedElementDataStorageInterface& DataStorage) const {}
 	virtual void RegisterTickGroups(ITypedElementDataStorageInterface& DataStorage) const {}
 	virtual void RegisterQueries(ITypedElementDataStorageInterface& DataStorage) const {}
+
+	virtual void RegisterDealiaser(ITypedElementDataStorageCompatibilityInterface& DataStorageCompatibility) const {}
 	
 	virtual void RegisterWidgetPurposes(ITypedElementDataStorageUiInterface& DataStorageUi) const {}
 	virtual void RegisterWidgetConstructors(ITypedElementDataStorageInterface& DataStorage,
