Added ability to filter out actors/objects from TEDS before they're created.

The following updates were made:
- It's now possible to register a filter function with the compatiblity layer that allows the registration of an AActor or UObject to be prevented.
- A previous update allowed the creation and registration of a row handle before the row is actually created. Many of the functions in the database however only checked for a valid row handle and not if the row was created. These have been updated to now check for a valid row before working to avoid problems later on.

#nojira

diff --git a/Engine/Plugins/Experimental/TypedElementsDataStorage/Source/TypedElementsDataStorage/Private/TypedElementDatabase.cpp b/Engine/Plugins/Experimental/TypedElementsDataStorage/Source/TypedElementsDataStorage/Private/TypedElementDatabase.cpp
--- a/Engine/Plugins/Experimental/TypedElementsDataStorage/Source/TypedElementsDataStorage/Private/TypedElementDatabase.cpp
+++ b/Engine/Plugins/Experimental/TypedElementsDataStorage/Source/TypedElementsDataStorage/Private/TypedElementDatabase.cpp
@@ -273,7 +273,7 @@ bool UTypedElementDatabase::HasRowBeenAssigned(TypedElementRowHandle Row) const
 bool UTypedElementDatabase::AddColumn(TypedElementRowHandle Row, const UScriptStruct* ColumnType)
 {
 	FMassEntityHandle Entity = FMassEntityHandle::FromNumber(Row);
-	if (ActiveEditorEntityManager && ActiveEditorEntityManager->IsEntityValid(Entity))
+	if (ActiveEditorEntityManager && ActiveEditorEntityManager->IsEntityActive(Entity))
 	{
 		if (ColumnType->IsChildOf(FMassTag::StaticStruct()))
 		{
@@ -303,7 +303,7 @@ bool UTypedElementDatabase::AddColumn(TypedElementRowHandle Row, FTopLevelAssetP
 void UTypedElementDatabase::RemoveColumn(TypedElementRowHandle Row, const UScriptStruct* ColumnType)
 {
 	FMassEntityHandle Entity = FMassEntityHandle::FromNumber(Row);
-	if (ActiveEditorEntityManager && ActiveEditorEntityManager->IsEntityValid(Entity))
+	if (ActiveEditorEntityManager && ActiveEditorEntityManager->IsEntityActive(Entity))
 	{
 		if (ColumnType->IsChildOf(FMassTag::StaticStruct()))
 		{
@@ -328,7 +328,7 @@ void UTypedElementDatabase::RemoveColumn(TypedElementRowHandle Row, FTopLevelAss
 void* UTypedElementDatabase::AddOrGetColumnData(TypedElementRowHandle Row, const UScriptStruct* ColumnType)
 {
 	FMassEntityHandle Entity = FMassEntityHandle::FromNumber(Row);
-	if (ActiveEditorEntityManager && ActiveEditorEntityManager->IsEntityValid(Entity) &&
+	if (ActiveEditorEntityManager && ActiveEditorEntityManager->IsEntityActive(Entity) &&
 		ColumnType && ColumnType->IsChildOf(FMassFragment::StaticStruct()))
 	{
 		FStructView Column = ActiveEditorEntityManager->GetFragmentDataStruct(Entity, ColumnType);
@@ -371,7 +371,7 @@ ColumnDataResult UTypedElementDatabase::AddOrGetColumnData(TypedElementRowHandle
 void* UTypedElementDatabase::GetColumnData(TypedElementRowHandle Row, const UScriptStruct* ColumnType)
 {
 	FMassEntityHandle Entity = FMassEntityHandle::FromNumber(Row);
-	if (ActiveEditorEntityManager && ActiveEditorEntityManager->IsEntityValid(Entity) &&
+	if (ActiveEditorEntityManager && ActiveEditorEntityManager->IsEntityActive(Entity) &&
 		ColumnType && ColumnType->IsChildOf(FMassFragment::StaticStruct()))
 	{
 		FStructView Column = ActiveEditorEntityManager->GetFragmentDataStruct(Entity, ColumnType);
@@ -386,7 +386,7 @@ void* UTypedElementDatabase::GetColumnData(TypedElementRowHandle Row, const UScr
 ColumnDataResult UTypedElementDatabase::GetColumnData(TypedElementRowHandle Row, FTopLevelAssetPath ColumnName)
 {
 	FMassEntityHandle Entity = FMassEntityHandle::FromNumber(Row);
-	if (ActiveEditorEntityManager && ActiveEditorEntityManager->IsEntityValid(Entity))
+	if (ActiveEditorEntityManager && ActiveEditorEntityManager->IsEntityActive(Entity))
 	{
 		const UScriptStruct* FragmentType = nullptr;
 		FMassArchetypeHandle Archetype = ActiveEditorEntityManager->GetArchetypeForEntityUnsafe(Entity);
@@ -414,7 +414,7 @@ ColumnDataResult UTypedElementDatabase::GetColumnData(TypedElementRowHandle Row,
 bool UTypedElementDatabase::AddColumns(TypedElementRowHandle Row, TConstArrayView<const UScriptStruct*> Columns)
 {
 	FMassEntityHandle Entity = FMassEntityHandle::FromNumber(Row);
-	if (ActiveEditorEntityManager && ActiveEditorEntityManager->IsEntityValid(Entity))
+	if (ActiveEditorEntityManager && ActiveEditorEntityManager->IsEntityActive(Entity))
 	{
 		FMassArchetypeHandle Archetype = ActiveEditorEntityManager->GetArchetypeForEntity(Entity);
 
@@ -434,7 +434,7 @@ bool UTypedElementDatabase::AddColumns(TypedElementRowHandle Row, TConstArrayVie
 void UTypedElementDatabase::RemoveColumns(TypedElementRowHandle Row, TConstArrayView<const UScriptStruct*> Columns)
 {
 	FMassEntityHandle Entity = FMassEntityHandle::FromNumber(Row);
-	if (ActiveEditorEntityManager && ActiveEditorEntityManager->IsEntityValid(Entity))
+	if (ActiveEditorEntityManager && ActiveEditorEntityManager->IsEntityActive(Entity))
 	{
 		FMassArchetypeHandle Archetype = ActiveEditorEntityManager->GetArchetypeForEntity(Entity);
 
@@ -454,7 +454,7 @@ bool UTypedElementDatabase::AddRemoveColumns(TypedElementRowHandle Row,
 {
 	bool bResult = false;
 	FMassEntityHandle Entity = FMassEntityHandle::FromNumber(Row);
-	if (ActiveEditorEntityManager && ActiveEditorEntityManager->IsEntityValid(Entity))
+	if (ActiveEditorEntityManager && ActiveEditorEntityManager->IsEntityActive(Entity))
 	{
 		FMassArchetypeHandle Archetype = ActiveEditorEntityManager->GetArchetypeForEntity(Entity);
 
@@ -506,7 +506,7 @@ bool UTypedElementDatabase::BatchAddRemoveColumns(TConstArrayView<TypedElementRo
 			for (TypedElementRowHandle EntityId : Rows)
 			{
 				FMassEntityHandle Entity = FMassEntityHandle::FromNumber(EntityId);
-				if (ActiveEditorEntityManager->IsEntityValid(Entity))
+				if (ActiveEditorEntityManager->IsEntityActive(Entity))
 				{
 					FMassArchetypeHandle Archetype = ActiveEditorEntityManager->GetArchetypeForEntity(Entity);
 					EntityHandleArray& EntityCollection = LookupTable.FindOrAdd(Archetype);
@@ -540,7 +540,7 @@ bool UTypedElementDatabase::BatchAddRemoveColumns(TConstArrayView<TypedElementRo
 bool UTypedElementDatabase::HasColumns(TypedElementRowHandle Row, TConstArrayView<const UScriptStruct*> ColumnTypes) const
 {
 	FMassEntityHandle Entity = FMassEntityHandle::FromNumber(Row);
-	if (ActiveEditorEntityManager && ActiveEditorEntityManager->IsEntityValid(Entity))
+	if (ActiveEditorEntityManager && ActiveEditorEntityManager->IsEntityActive(Entity))
 	{
 		FMassArchetypeHandle Archetype = ActiveEditorEntityManager->GetArchetypeForEntity(Entity);
 		const FMassArchetypeCompositionDescriptor& Composition = ActiveEditorEntityManager->GetArchetypeComposition(Archetype);
@@ -571,7 +571,7 @@ bool UTypedElementDatabase::HasColumns(TypedElementRowHandle Row, TConstArrayVie
 bool UTypedElementDatabase::HasColumns(TypedElementRowHandle Row, TConstArrayView<TWeakObjectPtr<const UScriptStruct>> ColumnTypes) const
 {
 	FMassEntityHandle Entity = FMassEntityHandle::FromNumber(Row);
-	if (ActiveEditorEntityManager && ActiveEditorEntityManager->IsEntityValid(Entity))
+	if (ActiveEditorEntityManager && ActiveEditorEntityManager->IsEntityActive(Entity))
 	{
 		FMassArchetypeHandle Archetype = ActiveEditorEntityManager->GetArchetypeForEntity(Entity);
 		const FMassArchetypeCompositionDescriptor& Composition = ActiveEditorEntityManager->GetArchetypeComposition(Archetype);
diff --git a/Engine/Plugins/Experimental/TypedElementsDataStorage/Source/TypedElementsDataStorage/Private/TypedElementDatabaseCompatibility.cpp b/Engine/Plugins/Experimental/TypedElementsDataStorage/Source/TypedElementsDataStorage/Private/TypedElementDatabaseCompatibility.cpp
--- a/Engine/Plugins/Experimental/TypedElementsDataStorage/Source/TypedElementsDataStorage/Private/TypedElementDatabaseCompatibility.cpp
+++ b/Engine/Plugins/Experimental/TypedElementsDataStorage/Source/TypedElementsDataStorage/Private/TypedElementDatabaseCompatibility.cpp
@@ -35,6 +35,11 @@ void UTypedElementDatabaseCompatibility::Deinitialize()
 	Reset();
 }
 
+void UTypedElementDatabaseCompatibility::RegisterRegistrationFilter(ObjectRegistrationFilter Filter)
+{
+	ObjectRegistrationFilters.Add(MoveTemp(Filter));
+}
+
 void UTypedElementDatabaseCompatibility::RegisterDealiaserCallback(ObjectToRowDealiaser Dealiaser)
 {
 	ObjectToRowDialiasers.Add(MoveTemp(Dealiaser));
@@ -53,7 +58,8 @@ TypedElementRowHandle UTypedElementDatabaseCompatibility::AddCompatibleObjectExp
 	}
 	else
 	{
-		if (ensureMsgf(Storage, TEXT("Trying to add a UObject to Typed Element's Data Storage before the storage is available.")))
+		if (ensureMsgf(Storage, TEXT("Trying to add a UObject to Typed Element's Data Storage before the storage is available.")) && 
+			ShouldAddObject(Object))
 		{
 			TypedElementRowHandle ReservedRow = Storage->ReserveRow();
 			ReverseObjectLookup.Add(Object, ReservedRow);
@@ -77,7 +83,8 @@ TypedElementRowHandle UTypedElementDatabaseCompatibility::AddCompatibleObjectExp
 
 TypedElementRowHandle UTypedElementDatabaseCompatibility::AddCompatibleObjectExplicit(AActor* Actor, TypedElementTableHandle Table)
 {
-	if (ensureMsgf(Storage, TEXT("Trying to add an actor to Typed Element's Data Storage before the storage is available.")))
+	if (ensureMsgf(Storage, TEXT("Trying to add an actor to Typed Element's Data Storage before the storage is available.")) &&
+		ShouldAddObject(Actor))
 	{
 		// Registration is delayed for two reasons:
 		//	1. Allows entity creation in a single batch rather than multiple individual additions.
@@ -246,6 +253,18 @@ void UTypedElementDatabaseCompatibility::CreateStandardArchetypes()
 		FName("Editor_StandardExternalObjectTable"));
 }
 
+bool UTypedElementDatabaseCompatibility::ShouldAddObject(const UObject* Object) const
+{
+	bool Include = true;
+	const ObjectRegistrationFilter* Filter = ObjectRegistrationFilters.GetData();
+	const ObjectRegistrationFilter* FilterEnd = Filter + ObjectRegistrationFilters.Num();
+	for (; Include && Filter != FilterEnd; ++Filter)
+	{
+		Include = (*Filter)(*this, Object);
+	}
+	return Include;
+}
+
 TypedElementRowHandle UTypedElementDatabaseCompatibility::DealiasObject(const UObject* Object) const
 {
 	for (const ObjectToRowDealiaser& Dealiaser : ObjectToRowDialiasers)
diff --git a/Engine/Plugins/Experimental/TypedElementsDataStorage/Source/TypedElementsDataStorage/Private/TypedElementDatabaseCompatibility.h b/Engine/Plugins/Experimental/TypedElementsDataStorage/Source/TypedElementsDataStorage/Private/TypedElementDatabaseCompatibility.h
--- a/Engine/Plugins/Experimental/TypedElementsDataStorage/Source/TypedElementsDataStorage/Private/TypedElementDatabaseCompatibility.h
+++ b/Engine/Plugins/Experimental/TypedElementsDataStorage/Source/TypedElementsDataStorage/Private/TypedElementDatabaseCompatibility.h
@@ -26,6 +26,7 @@ public:
 	void Initialize(ITypedElementDataStorageInterface* StorageInterface);
 	void Deinitialize();
 
+	void RegisterRegistrationFilter(ObjectRegistrationFilter Filter) override;
 	void RegisterDealiaserCallback(ObjectToRowDealiaser Dealiaser) override;
 
 	TypedElementRowHandle AddCompatibleObjectExplicit(UObject* Object) override;
@@ -48,6 +49,7 @@ private:
 	void Reset();
 	void CreateStandardArchetypes();
 	
+	bool ShouldAddObject(const UObject* Object) const;
 	TypedElementRowHandle DealiasObject(const UObject* Object) const;
 
 	void Tick();
@@ -87,6 +89,7 @@ private:
 	TMap<TypedElementTableHandle, PendingRegistration<TWeakObjectPtr<UObject>>> UObjectsPendingRegistration;
 	TMap<TypedElementTableHandle, PendingRegistration<ExternalObjectRegistration>> ExternalObjectsPendingRegistration;
 	
+	TArray<ObjectRegistrationFilter> ObjectRegistrationFilters;
 	TArray<ObjectToRowDealiaser> ObjectToRowDialiasers;
 
 	TypedElementTableHandle StandardActorTable{ TypedElementInvalidTableHandle };
diff --git a/Engine/Plugins/Experimental/TypedElementsDataStorage/Source/TypedElementsDataStorage/Private/TypedElementsDataStorage.cpp b/Engine/Plugins/Experimental/TypedElementsDataStorage/Source/TypedElementsDataStorage/Private/TypedElementsDataStorage.cpp
--- a/Engine/Plugins/Experimental/TypedElementsDataStorage/Source/TypedElementsDataStorage/Private/TypedElementsDataStorage.cpp
+++ b/Engine/Plugins/Experimental/TypedElementsDataStorage/Source/TypedElementsDataStorage/Private/TypedElementsDataStorage.cpp
@@ -97,6 +97,7 @@ void FTypedElementsDataStorageModule::StartupModule()
 				{
 					Factory->RegisterTables(*Database);
 					Factory->RegisterTickGroups(*Database);
+					Factory->RegisterRegistrationFilters(*DatabaseCompatibility);
 					Factory->RegisterDealiaser(*DatabaseCompatibility);
 					Factory->RegisterWidgetPurposes(*DatabaseUi);
 				}
diff --git a/Engine/Source/Runtime/TypedElementFramework/Public/Elements/Interfaces/TypedElementDataStorageCompatibilityInterface.h b/Engine/Source/Runtime/TypedElementFramework/Public/Elements/Interfaces/TypedElementDataStorageCompatibilityInterface.h
--- a/Engine/Source/Runtime/TypedElementFramework/Public/Elements/Interfaces/TypedElementDataStorageCompatibilityInterface.h
+++ b/Engine/Source/Runtime/TypedElementFramework/Public/Elements/Interfaces/TypedElementDataStorageCompatibilityInterface.h
@@ -30,6 +30,7 @@ class ITypedElementDataStorageCompatibilityInterface
 	GENERATED_BODY()
 
 public:
+	using ObjectRegistrationFilter = TFunction<bool(const ITypedElementDataStorageCompatibilityInterface&, const UObject*)>;
 	using ObjectToRowDealiaser = TFunction<TypedElementRowHandle(const ITypedElementDataStorageCompatibilityInterface&, const UObject*)>;
 
 	/**
@@ -57,7 +58,17 @@ public:
 	TypedElementRowHandle FindRowWithCompatibleObject(ObjectType&& Object) const;
 
 	/**
-	 * @section Alias functions
+	 * @section Callback registration
+	 * Functions to register callbacks with the compatibility layer to help refine its operations.
+	 */
+	 
+	/**
+	 * Objects like actors are registered through the compatibility layer in bulk. This can lead to objects being added that cause
+	 * conflicts with other data in the data storage. This callback offers the opportunity to inspect the objects that are being
+	 * added and if they include an object that shouldn't be store it can filter them out.
+	 */
+	virtual void RegisterRegistrationFilter(ObjectRegistrationFilter Filter) = 0;
+	/**
 	 * Notifications and request can be made to the compatibility layer for objects that are stored but don't directly map to a row.
 	 * An example is a UObject represented by a column. If the UObject gets updated there's no direct mapping to the row the column is
 	 * stored in but the row still needs to be updated. For cases like this it's possible to store information to find the row that's
diff --git a/Engine/Source/Runtime/TypedElementFramework/Public/Elements/Interfaces/TypedElementDataStorageFactory.h b/Engine/Source/Runtime/TypedElementFramework/Public/Elements/Interfaces/TypedElementDataStorageFactory.h
--- a/Engine/Source/Runtime/TypedElementFramework/Public/Elements/Interfaces/TypedElementDataStorageFactory.h
+++ b/Engine/Source/Runtime/TypedElementFramework/Public/Elements/Interfaces/TypedElementDataStorageFactory.h
@@ -30,6 +30,7 @@ public:
 	virtual void RegisterTickGroups(ITypedElementDataStorageInterface& DataStorage) const {}
 	virtual void RegisterQueries(ITypedElementDataStorageInterface& DataStorage) const {}
 
+	virtual void RegisterRegistrationFilters(ITypedElementDataStorageCompatibilityInterface& DataStorageCompatibility) const {}
 	virtual void RegisterDealiaser(ITypedElementDataStorageCompatibilityInterface& DataStorageCompatibility) const {}
 	
 	virtual void RegisterWidgetPurposes(ITypedElementDataStorageUiInterface& DataStorageUi) const {}
