Use the struct flags to determine if DestroyStruct() needs to be called on TypedElementColumns that are temporarily in the ScratchBuffer
* DestroyStruct internally calls the CPP define destructor if it is possible.  It also calls the property destructor chain for properties that do not get caught by the CPP destruction function.
* This CL solves a nullptr dereference if using runtime generated UScriptStruct (ie. UPropertyBag) since they don't have CppOps and will ensure that Properties are properly cleaned up.
* Aside, there isn't a way to define an ObjectMoveOperator for these non-cpp types easily, but changed the code anyways when adding to scratch buffer in the AddColumnUninitialize overload with ObjectMoveOperator to keep consistency

diff --git a/Engine/Plugins/Experimental/TypedElementsDataStorage/Source/TypedElementsDataStorage/Private/Processors/TypedElementProcessorAdaptors.cpp b/Engine/Plugins/Experimental/TypedElementsDataStorage/Source/TypedElementsDataStorage/Private/Processors/TypedElementProcessorAdaptors.cpp
--- a/Engine/Plugins/Experimental/TypedElementsDataStorage/Source/TypedElementsDataStorage/Private/Processors/TypedElementProcessorAdaptors.cpp
+++ b/Engine/Plugins/Experimental/TypedElementsDataStorage/Source/TypedElementsDataStorage/Private/Processors/TypedElementProcessorAdaptors.cpp
@@ -262,9 +262,15 @@ public:
 		};
 
 		void* ObjectCopy = ScratchBuffer.Allocate(ObjectType->GetStructureSize(), ObjectType->GetMinAlignment());
-		FAddValueColumn* AddedColumn = ObjectType->GetCppStructOps()->HasDestructor() ?
-			ScratchBuffer.Emplace<FAddValueColumnWithDestructor>(ObjectType, FMassEntityHandle::FromNumber(Row), ObjectCopy) :
-			ScratchBuffer.Emplace<FAddValueColumn>(ObjectType, FMassEntityHandle::FromNumber(Row), ObjectCopy);
+		FAddValueColumn* AddedColumn = nullptr;
+		if (ObjectType->StructFlags & (STRUCT_IsPlainOldData | STRUCT_NoDestructor))
+		{
+			AddedColumn = ScratchBuffer.Emplace<FAddValueColumn>(ObjectType, FMassEntityHandle::FromNumber(Row), ObjectCopy);
+		}
+		else
+		{
+			AddedColumn = ScratchBuffer.Emplace<FAddValueColumnWithDestructor>(ObjectType, FMassEntityHandle::FromNumber(Row), ObjectCopy);
+		}
 
 		Context.Defer().PushCommand<FMassDeferredAddCommand>(
 			[AddedColumn](FMassEntityManager& System)
@@ -314,9 +320,15 @@ public:
 		};
 
 		void* MovedObject = ScratchBuffer.Allocate(ObjectType->GetStructureSize(), ObjectType->GetMinAlignment());
-		FAddMoveableValueColumn* AddedColumn = ObjectType->GetCppStructOps()->HasDestructor() ?
-			ScratchBuffer.Emplace<FAddMoveableValueColumnWithDestructor>(Mover, ObjectType, FMassEntityHandle::FromNumber(Row), MovedObject) :
-			ScratchBuffer.Emplace<FAddMoveableValueColumn>(Mover, ObjectType, FMassEntityHandle::FromNumber(Row), MovedObject);
+		FAddMoveableValueColumn* AddedColumn = nullptr;
+		if (ObjectType->StructFlags & (STRUCT_IsPlainOldData | STRUCT_NoDestructor))
+		{
+			AddedColumn = ScratchBuffer.Emplace<FAddMoveableValueColumn>(Mover, ObjectType, FMassEntityHandle::FromNumber(Row), MovedObject);
+		}
+		else
+		{
+			AddedColumn = ScratchBuffer.Emplace<FAddMoveableValueColumnWithDestructor>(Mover, ObjectType, FMassEntityHandle::FromNumber(Row), MovedObject);
+		}
 
 		Context.Defer().PushCommand<FMassDeferredAddCommand>(
 			[AddedColumn](FMassEntityManager& System)
