Disable TEDS Processors from pruning on missing Mass archtypes
* TypedElementDataStore intends to dynamically register processors fairly frequently - it will not always be possible to register archtypes up front.
* Therefore, the processors that are part of the TypedElementDataStore will not participate in the pruning behaviour when it's query does not intersect with any explicit table schema.  The table may be generated at a later time.

diff --git a/Engine/Plugins/Experimental/TypedElementsDataStorage/Source/TypedElementsDataStorage/Private/Processors/TypedElementProcessorAdaptors.cpp b/Engine/Plugins/Experimental/TypedElementsDataStorage/Source/TypedElementsDataStorage/Private/Processors/TypedElementProcessorAdaptors.cpp
--- a/Engine/Plugins/Experimental/TypedElementsDataStorage/Source/TypedElementsDataStorage/Private/Processors/TypedElementProcessorAdaptors.cpp
+++ b/Engine/Plugins/Experimental/TypedElementsDataStorage/Source/TypedElementsDataStorage/Private/Processors/TypedElementProcessorAdaptors.cpp
@@ -757,6 +757,15 @@ bool UTypedElementQueryProcessorCallbackAdapterProcessorBase::ConfigureQueryCall
 	return ConfigureQueryCallbackData(Query, QueryStore, ScratchBuffer, {});
 }
 
+bool UTypedElementQueryProcessorCallbackAdapterProcessorBase::ShouldAllowQueryBasedPruning(
+	const bool bRuntimeMode) const
+{
+	// TEDS is much more dynamic with when tables and processors are added and removed
+	// Don't prune processors if they have queries where no table is defined, it is possible
+	// the table will be dynamically created later.
+	return false;
+}
+
 bool UTypedElementQueryProcessorCallbackAdapterProcessorBase::ConfigureQueryCallbackData(
 	FTypedElementExtendedQuery& Query,
 	FTypedElementExtendedQueryStore& QueryStore, 
diff --git a/Engine/Plugins/Experimental/TypedElementsDataStorage/Source/TypedElementsDataStorage/Private/Processors/TypedElementProcessorAdaptors.h b/Engine/Plugins/Experimental/TypedElementsDataStorage/Source/TypedElementsDataStorage/Private/Processors/TypedElementProcessorAdaptors.h
--- a/Engine/Plugins/Experimental/TypedElementsDataStorage/Source/TypedElementsDataStorage/Private/Processors/TypedElementProcessorAdaptors.h
+++ b/Engine/Plugins/Experimental/TypedElementsDataStorage/Source/TypedElementsDataStorage/Private/Processors/TypedElementProcessorAdaptors.h
@@ -92,6 +92,8 @@ public:
 		FTypedElementExtendedQueryStore& QueryStore,
 		FTypedElementDatabaseScratchBuffer& ScratchBuffer);
 
+	virtual bool ShouldAllowQueryBasedPruning(const bool bRuntimeMode) const override;
+
 protected:
 	bool ConfigureQueryCallbackData(
 		FTypedElementExtendedQuery& Query, 
