Fixed an issue with stale object references in TEDS.

Objects weren't always correctly removed from the reverse lookup list in the TEDS' compatibility layer. This has been fixed.

diff --git a/Engine/Plugins/Experimental/TypedElementsDataStorage/Source/TypedElementsDataStorage/Private/TypedElementDatabaseCompatibility.cpp b/Engine/Plugins/Experimental/TypedElementsDataStorage/Source/TypedElementsDataStorage/Private/TypedElementDatabaseCompatibility.cpp
--- a/Engine/Plugins/Experimental/TypedElementsDataStorage/Source/TypedElementsDataStorage/Private/TypedElementDatabaseCompatibility.cpp
+++ b/Engine/Plugins/Experimental/TypedElementsDataStorage/Source/TypedElementsDataStorage/Private/TypedElementDatabaseCompatibility.cpp
@@ -129,9 +129,10 @@ void UTypedElementDatabaseCompatibility::RemoveCompatibleObjectExplicit(UObject*
 	else
 	{
 		checkf(Storage, TEXT("Removing compatible objects is not supported before Typed Element's Database compatibility manager has been initialized."));
-		if (TypedElementRowHandle* Row = ReverseObjectLookup.Find(Object))
+		TypedElementRowHandle Row;
+		if (ReverseObjectLookup.RemoveAndCopyValue(Object, Row))
 		{
-			Storage->RemoveRow(*Row);
+			Storage->RemoveRow(Row);
 		}
 	}
 }
@@ -139,9 +140,10 @@ void UTypedElementDatabaseCompatibility::RemoveCompatibleObjectExplicit(UObject*
 void UTypedElementDatabaseCompatibility::RemoveCompatibleObjectExplicit(void* Object)
 {
 	checkf(Storage, TEXT("Removing compatible objects is not supported before Typed Element's Database compatibility manager has been initialized."));
-	if (TypedElementRowHandle* Row = ReverseObjectLookup.Find(Object))
+	TypedElementRowHandle Row;
+	if (ReverseObjectLookup.RemoveAndCopyValue(Object, Row))
 	{
-		Storage->RemoveRow(*Row);
+		Storage->RemoveRow(Row);
 	}
 }
 
