Fixed RWAccessDetector in TEDS ScratchBuffer
* The RWAccessDetector is in place to catch allocations occurring at the same time as the allocator's "RecycleBlocks".
* It is valid for allocations to occur from multiple threads though.
* The exising ScopedWriterDetector is an assertion that the scope should have unique access - however this isn't the intention of the function it is in.  It should have shared access which is denoted with a ScopedReaderAccessDetector.
* The fix is to change this to a ScopedReaderAccessDetector so that multiple threads can allocate from the scratch buffer at the same time.
* Aside: The names Reader and Writer might be the cause for this bug as an allocation somewhat feels like "writing".

diff --git a/Engine/Plugins/Experimental/TypedElementsDataStorage/Source/TypedElementsDataStorage/Private/TypedElementDatabaseScratchBuffer.cpp b/Engine/Plugins/Experimental/TypedElementsDataStorage/Source/TypedElementsDataStorage/Private/TypedElementDatabaseScratchBuffer.cpp
--- a/Engine/Plugins/Experimental/TypedElementsDataStorage/Source/TypedElementsDataStorage/Private/TypedElementDatabaseScratchBuffer.cpp
+++ b/Engine/Plugins/Experimental/TypedElementsDataStorage/Source/TypedElementsDataStorage/Private/TypedElementDatabaseScratchBuffer.cpp
@@ -29,7 +29,8 @@ FTypedElementDatabaseScratchBuffer::~FTypedElementDatabaseScratchBuffer()
 
 void* FTypedElementDatabaseScratchBuffer::Allocate(size_t Size, size_t Alignment)
 {
-	TScopedWriterDetector Detector(AccessDetector);
+	// ReaderAccessDetector to assert that this scope has shared access to any thread also making an Allocate call
+	TScopedReaderAccessDetector Detector(AccessDetector);
 	return GetThreadLocalBlockController().Allocate(Size, Alignment);
 }
 
@@ -37,6 +38,7 @@ void FTypedElementDatabaseScratchBuffer::RecycleBlocks()
 {
 	if (FullBlocks)
 	{
+		// WriterAccessDetector to assert that this scope has exclusive access - should not overlap Allocate()
 		TScopedWriterDetector Detector(AccessDetector);
 		
 		// Destroy any objects in blocks. It walks the allocations in a block backwards and invokes the destructors on the allocated 
