First iteration of TEDS memento system

This CL introduces a mechanism to opt into creating mementos of columns for mementoized rows.
 Mementoization is the process of copying the a subset of a row's state and storing it to be later read to restore the state of a row.

The heart of this process is in the MementoTranslator.  These are CDOs which define the mapping between the column that will be mementoized and the memento itself.  They also define the mechanism to covert between columns and mementos.  This initial implementation only allows for column to memento operations, though extension to allow the reverse will be added later.

MementoTranslators have two levels to the API.
Level 1: The DefaultMementoTranslator is the easiest to set up.  It will generate a runtime struct of the mementoized column based on the UPROPERTIES of the column and has code that will automatically populate the memento from the column.  There is an example of it's use in UTypedElementLocalTransformColumnMementoProcessor.
MementoTranslator:
Level 2: The MementoTranslator can be overridden by user code to get more control over the mementoization process.  It allows explicitly defining the translation process. This can be useful for optimizing the size of the mementos themselves or the speed at which mementoization occurs.  The UTypedElementPackagePathMementoTranslator is an example of using this API.

Rows which want to opt into mementoization need to add a FTypedElementMementoOnDelete column and a placeholder unpopulated Memento row.  This MementoSystem will watch for rows deleted with this column and any mementoized column defined by the available MementoTranslators and invoke the appropriate mementoization operators.

Memento rows will exist during the lifetime of a mementoized row and will be destroyed automatically after 120 frames to avoid continued expansion of the memento tables.

By default, mementoization is disabled for objects in the Compatibility layer.

CVARS:
`teds.mementos.enable 1` - enable mementoization on newly created objects.
`teds.mements.keepframes N` - The number of frames that a populated memento sticks around after the mementoized row is deleted

diff --git a/Engine/Plugins/Experimental/TypedElementsDataStorage/Source/TypedElementsDataStorage/Private/Memento/TypedElementMementoExamples.h b/Engine/Plugins/Experimental/TypedElementsDataStorage/Source/TypedElementsDataStorage/Private/Memento/TypedElementMementoExamples.h
new file mode 100644
--- /dev/null
+++ b/Engine/Plugins/Experimental/TypedElementsDataStorage/Source/TypedElementsDataStorage/Private/Memento/TypedElementMementoExamples.h
@@ -0,0 +1,55 @@
+﻿// Copyright Epic Games, Inc. All Rights Reserved.
+
+#pragma once
+
+#include "Elements/Columns/TypedElementPackageColumns.h"
+#include "Elements/Columns/TypedElementTransformColumns.h"
+#include "Memento/TypedElementMementoTranslators.h"
+#include "TypedElementMementoExamples.generated.h"
+
+/**
+ * Default LocalTransform MementoTranslator
+ * Implementation will generate a memento structure to put into TEDS at runtime based on UPROPERTY declarations
+ */
+UCLASS()
+class UTypedElementLocalTransformColumnMementoTranslator final : public UTypedElementDefaultMementoTranslator
+{
+	GENERATED_BODY()
+public:
+	virtual const UScriptStruct* GetColumnType() const override { return FTypedElementLocalTransformColumn::StaticStruct();	}
+};
+
+/**
+ * Explicit memento for PackagePath
+ */
+USTRUCT()
+struct FTypedElementPackagePathColumnMemento final : public FTypedElementDataStorageColumn
+{
+	GENERATED_BODY()
+	
+	UPROPERTY()
+	FName Path;
+};
+
+/**
+ * Explicit PackagePath MementoTranslator
+ */
+UCLASS()
+class UTypedElementPackagePathMementoTranslator final : public UTypedElementMementoTranslatorBase
+{
+	GENERATED_BODY()
+public:
+	virtual const UScriptStruct* GetColumnType() const override { return FTypedElementPackagePathColumn::StaticStruct(); }
+	virtual const UScriptStruct* GetMementoType() const override { return FTypedElementPackagePathColumnMemento::StaticStruct(); }
+
+private:
+	virtual void TranslateColumnToMemento(const void* TypedErasedColumn, void* TypeErasedMemento) const override
+	{
+		// Boilerplate to get the types back
+		const FTypedElementPackagePathColumn* Column = static_cast<const FTypedElementPackagePathColumn*>(TypedErasedColumn);
+		FTypedElementPackagePathColumnMemento* Memento = static_cast<FTypedElementPackagePathColumnMemento*>(TypeErasedMemento);
+
+		// Explicit mapping between the FString 
+		Memento->Path = FName(Column->Path);
+	}
+};
\ No newline at end of file
diff --git a/Engine/Plugins/Experimental/TypedElementsDataStorage/Source/TypedElementsDataStorage/Private/Memento/TypedElementMementoInterface.cpp b/Engine/Plugins/Experimental/TypedElementsDataStorage/Source/TypedElementsDataStorage/Private/Memento/TypedElementMementoInterface.cpp
new file mode 100644
--- /dev/null
+++ b/Engine/Plugins/Experimental/TypedElementsDataStorage/Source/TypedElementsDataStorage/Private/Memento/TypedElementMementoInterface.cpp
@@ -0,0 +1,11 @@
+﻿// Copyright Epic Games, Inc. All Rights Reserved.
+
+#include "Memento/TypedElementMementoInterface.h"
+
+#include "TypedElementMementoSystem.h"
+
+TypedElementTableHandle UTypedElementMementoInterface::GetUnpopulatedMementoTable()
+{
+	const UTypedElementMementoSystemFactory* System = GetDefault<UTypedElementMementoSystemFactory>();
+	return System->GetUnpopulatedMementoTable();
+}
diff --git a/Engine/Plugins/Experimental/TypedElementsDataStorage/Source/TypedElementsDataStorage/Private/Memento/TypedElementMementoRowTypes.h b/Engine/Plugins/Experimental/TypedElementsDataStorage/Source/TypedElementsDataStorage/Private/Memento/TypedElementMementoRowTypes.h
new file mode 100644
--- /dev/null
+++ b/Engine/Plugins/Experimental/TypedElementsDataStorage/Source/TypedElementsDataStorage/Private/Memento/TypedElementMementoRowTypes.h
@@ -0,0 +1,31 @@
+﻿// Copyright Epic Games, Inc. All Rights Reserved.
+
+#pragma once
+
+#include "Elements/Interfaces/TypedElementDataStorageInterface.h"
+#include "TypedElementMementoRowTypes.generated.h"
+
+/**
+ * @file MementoRowTypes.h
+ * Column/Tags used internally for the memento system
+ */
+
+/**
+ * MementoTag denotes that the row is a memento
+ */
+USTRUCT()
+struct FTypedElementMementoTag : public FTypedElementDataStorageTag
+{
+	GENERATED_BODY()
+};
+
+/**
+ * Column for a memento row indicating that the row is populated with memento data
+ */
+USTRUCT()
+struct FTypedElementMementoPopulated : public FTypedElementDataStorageColumn
+{
+	// Frame number that the memento row became populated
+	GENERATED_BODY()
+	uint64 FrameNumber;
+};
\ No newline at end of file
diff --git a/Engine/Plugins/Experimental/TypedElementsDataStorage/Source/TypedElementsDataStorage/Private/Memento/TypedElementMementoSystem.cpp b/Engine/Plugins/Experimental/TypedElementsDataStorage/Source/TypedElementsDataStorage/Private/Memento/TypedElementMementoSystem.cpp
new file mode 100644
--- /dev/null
+++ b/Engine/Plugins/Experimental/TypedElementsDataStorage/Source/TypedElementsDataStorage/Private/Memento/TypedElementMementoSystem.cpp
@@ -0,0 +1,139 @@
+﻿// Copyright Epic Games, Inc. All Rights Reserved.
+
+#include "TypedElementMementoSystem.h"
+
+#include "Elements/Framework/TypedElementQueryBuilder.h"
+#include "HAL/IConsoleManager.h"
+#include "Memento/TypedElementMementoInterface.h"
+#include "Memento/TypedElementMementoTranslators.h"
+#include "PropertyBag.h"
+#include "TypedElementMementoRowTypes.h"
+
+namespace Private
+{
+	// Number of frames to keep memento rows before deletion
+	int32 GMementoKeepFrames = 120;
+	FAutoConsoleVariableRef CVarMementoKeepFrame(
+		TEXT("teds.mementos.keepframes"),
+		GMementoKeepFrames,
+		TEXT("Number of frames to keep memento rows before deletion\n"));
+}
+
+void UTypedElementMementoSystemFactory::RegisterTables(ITypedElementDataStorageInterface& DataStorage) const
+{
+	MementoRowBaseTable = DataStorage.RegisterTable(
+		{FTypedElementMementoTag::StaticStruct()},
+		TEXT("MementoRowBaseTable"));
+
+	// Memento rows become populated when the FTypedElementMementoPopulated column is added
+	TypedElementTableHandle MementoRowBasePopulatedTable = DataStorage.RegisterTable(
+		MementoRowBaseTable,
+		{FTypedElementMementoPopulated::StaticStruct()},
+		TEXT("MementoRowBasePopulatedTable"));
+
+	// Note: Need to register PopulatedMementTable so that "Add Memento Deletion Policy Data" processor does not get pruned
+	// This is because Mass prunes processors that register queries which don't have a matching archtype
+	// at processor construction time
+}
+
+void UTypedElementMementoSystemFactory::RegisterQueries(ITypedElementDataStorageInterface& DataStorage) const
+{
+	TArray<const UTypedElementMementoTranslatorBase*> MementoTranslators;
+
+	// Discover all MementoTranslators
+	{
+		const bool bIncludeDerived = true;
+		EObjectFlags ExcludeFlags = EObjectFlags::RF_NoFlags;
+		ForEachObjectOfClass(UTypedElementMementoTranslatorBase::StaticClass(), [&MementoTranslators](UObject* Object)
+		{
+			const UTypedElementMementoTranslatorBase* TranslatorCandidate = Cast<UTypedElementMementoTranslatorBase>(Object);
+			// Exclude abstract classes
+			if (TranslatorCandidate->GetClass()->GetClassFlags() & EClassFlags::CLASS_Abstract)
+			{
+				return;
+			}
+			MementoTranslators.Add(TranslatorCandidate);
+		},
+		bIncludeDerived,
+		ExcludeFlags);
+	}
+
+	using namespace TypedElementQueryBuilder;
+
+	// Setup observer queries to execute memento translators on row deletions that match the translator's
+	// capabilities.
+	for (int32 Index = 0, End = MementoTranslators.Num(); Index < End; ++Index)
+	{
+		const UTypedElementMementoTranslatorBase* MementoTranslator = MementoTranslators[Index];
+		const UScriptStruct* MementoizedColumn = MementoTranslator->GetColumnType();
+		const UScriptStruct* Memento = MementoTranslator->GetMementoType();
+		
+		const FName TranslationProcessorName = FName(FString::Printf(TEXT("MementoTranslator: %s -> %s"), *MementoizedColumn->GetName(), *Memento->GetName()));
+		const TypedElementQueryHandle QueryHandle = DataStorage.RegisterQuery(
+			Select(
+				TranslationProcessorName,
+				FObserver::OnRemove<FTypedElementMementoOnDelete>(),
+				[MementoTranslator, Memento, MementoizedColumn](TypedElementDataStorage::IQueryContext& Context, TypedElementRowHandle row, const FTypedElementMementoOnDelete& MementoRow)
+				{
+					void* MementoObject = Context.AddColumnUnitialized(MementoRow.Memento, Memento);
+					Memento->InitializeStruct(MementoObject);
+		
+					const void* SourceColumn = Context.GetColumn(MementoizedColumn);
+					MementoTranslator->TranslateColumnToMemento(SourceColumn, MementoObject);
+				})
+				.ReadOnly(MementoizedColumn)
+				.Compile());
+		check(QueryHandle != TypedElementInvalidQueryHandle);
+	}
+
+	// Setup a default policy of deleting populated mementos after a set number of frames
+	// The first use-case of mementos are to support reinstancing which occurs over a single frame
+	// Therefore it should be safe to remove populated mementos after a couple of frames since the data
+	// is no longer needed.
+	// Expansion of memento deletion policies will be necessary for other use cases
+	{
+		const TypedElementQueryHandle QueryHandle = DataStorage.RegisterQuery(
+			Select(
+				TEXT("Add Memento Deletion Policy Data"),
+				FObserver::OnRemove<FTypedElementMementoOnDelete>(), // When a memento is populated
+				[](TypedElementDataStorage::IQueryContext& Context, TypedElementRowHandle Row, const FTypedElementMementoOnDelete& MementoRow)
+				{
+					Context.AddColumn<FTypedElementMementoPopulated>(
+						MementoRow.Memento,
+						FTypedElementMementoPopulated
+						{
+							.FrameNumber = GFrameCounter
+						});
+				})
+				.Compile());
+		check(QueryHandle != TypedElementInvalidQueryHandle);
+	}
+
+	using DSI = ITypedElementDataStorageInterface;
+
+	/**
+	 * A processor which deletes mementos that were populated more than N frames ago
+	 */
+	{
+		const TypedElementQueryHandle QueryHandle = DataStorage.RegisterQuery(
+			Select(
+			TEXT("Delete populated mementos older than N frames"),
+			FProcessor(DSI::EQueryTickPhase::FrameEnd, DataStorage.GetQueryTickGroupName(DSI::EQueryTickGroups::Default)),
+				[](TypedElementDataStorage::IQueryContext& Context, TypedElementRowHandle Row, const FTypedElementMementoPopulated& PopulateFrameNumber)
+				{
+					if (PopulateFrameNumber.FrameNumber + Private::GMementoKeepFrames < GFrameCounter)
+					{
+						Context.RemoveRow(Row);
+					}
+				})
+				.Where().All<FTypedElementMementoTag>()
+				.Compile()
+			);
+		check(QueryHandle != TypedElementInvalidQueryHandle);
+	}
+
+	
+}
+
+
+
diff --git a/Engine/Plugins/Experimental/TypedElementsDataStorage/Source/TypedElementsDataStorage/Private/Memento/TypedElementMementoSystem.h b/Engine/Plugins/Experimental/TypedElementsDataStorage/Source/TypedElementsDataStorage/Private/Memento/TypedElementMementoSystem.h
new file mode 100644
--- /dev/null
+++ b/Engine/Plugins/Experimental/TypedElementsDataStorage/Source/TypedElementsDataStorage/Private/Memento/TypedElementMementoSystem.h
@@ -0,0 +1,28 @@
+﻿// Copyright Epic Games, Inc. All Rights Reserved.
+
+#pragma once
+
+#include "Elements/Interfaces/TypedElementDataStorageFactory.h"
+#include "Elements/Columns/TypedElementSelectionColumns.h"
+
+#include "TypedElementMementoSystem.generated.h"
+
+
+UCLASS()
+class UTypedElementMementoSystemFactory : public UTypedElementDataStorageFactory
+{
+	GENERATED_BODY()
+
+public:
+	void RegisterQueries(ITypedElementDataStorageInterface& DataStorage) const override;
+	void RegisterTables(ITypedElementDataStorageInterface& DataStorage) const override;
+
+	TypedElementTableHandle GetUnpopulatedMementoTable() const;
+private:
+	mutable TypedElementTableHandle MementoRowBaseTable;
+};
+
+inline TypedElementTableHandle UTypedElementMementoSystemFactory::GetUnpopulatedMementoTable() const
+{
+	return MementoRowBaseTable;
+}
diff --git a/Engine/Plugins/Experimental/TypedElementsDataStorage/Source/TypedElementsDataStorage/Private/Memento/TypedElementMementoTranslators.cpp b/Engine/Plugins/Experimental/TypedElementsDataStorage/Source/TypedElementsDataStorage/Private/Memento/TypedElementMementoTranslators.cpp
new file mode 100644
--- /dev/null
+++ b/Engine/Plugins/Experimental/TypedElementsDataStorage/Source/TypedElementsDataStorage/Private/Memento/TypedElementMementoTranslators.cpp
@@ -0,0 +1,73 @@
+﻿// Copyright Epic Games, Inc. All Rights Reserved.
+
+#include "Memento/TypedElementMementoTranslators.h"
+
+#include "Elements/Interfaces/TypedElementDataStorageInterface.h"
+#include "PropertyBag.h"
+
+
+const UScriptStruct* UTypedElementDefaultMementoTranslator::GetMementoType() const
+{
+	return MementoType;
+}
+
+void UTypedElementDefaultMementoTranslator::PostInitProperties()
+{
+	Super::PostInitProperties();
+	
+	const UScriptStruct* SourceColumnType = GetColumnType();
+	if (SourceColumnType == nullptr)
+	{
+		return;
+	}
+
+	// Create a new runtime generated struct as the memento from the ColumnType based on the exposed UPROPERTIES
+	// and generate a mapping between the properties of the ColumnType to the Memento
+	// This mapping will be used to populate the memento and columns during translation
+	
+	TArray<FPropertyBagPropertyDesc> PropertyDescs;
+	for (FProperty* Property = SourceColumnType->PropertyLink; Property; Property = Property->PropertyLinkNext)
+	{
+		PropertyDescs.Add(FPropertyBagPropertyDesc(Property->GetFName(), Property));
+	}
+
+	// Note: The name of the memento will be an unhelpful "PropertyBag_<hash>"
+	// May wish to revisit property bag interface to allow name to be specified for debug purposes
+	const UPropertyBag* PropertyBag = UPropertyBag::GetOrCreateFromDescs(PropertyDescs);
+	MementoType = PropertyBag;
+
+	// Need to change the type to a FTypedElementDataStorageColumn to appease TEDS/Mass
+	const_cast<UPropertyBag*>(PropertyBag)->SetSuperStruct(FTypedElementDataStorageColumn::StaticStruct());
+
+	// Create the property mapping
+	for (FProperty* SourceProperty = SourceColumnType->PropertyLink; SourceProperty; SourceProperty = SourceProperty->PropertyLinkNext)
+	{
+		const FProperty* DestinationProperty = PropertyBag->FindPropertyByName(SourceProperty->GetFName());
+		if (DestinationProperty && DestinationProperty->SameType(SourceProperty))
+		{
+			SourceProperties.Add(SourceProperty);
+			DestinationProperties.Add(DestinationProperty);
+		}
+	}
+}
+
+void UTypedElementDefaultMementoTranslator::TranslateColumnToMemento(const void* Column, void* Memento) const
+{
+	const UScriptStruct* ColumnType = GetColumnType();
+
+	const std::byte* BaseAddressColumn = static_cast<const std::byte*>(Column);
+	std::byte* BaseAddressMemento = static_cast<std::byte*>(Memento);
+
+	check(SourceProperties.Num() == DestinationProperties.Num());
+
+	for (int32 PropertyIndex = 0, PropertyIndexEnd = SourceProperties.Num(); PropertyIndex < PropertyIndexEnd; ++PropertyIndex)
+	{
+		const FProperty* SourceProperty = SourceProperties[PropertyIndex];
+		const FProperty* DestinationProperty = DestinationProperties[PropertyIndex];
+		void* DestinationValueAddress = BaseAddressMemento + DestinationProperty->GetOffset_ForInternal();
+		const void* SourceValueAddress = BaseAddressColumn + SourceProperty->GetOffset_ForInternal();
+		SourceProperty->CopyCompleteValue(
+			DestinationValueAddress,
+			SourceValueAddress);
+	}	
+}
diff --git a/Engine/Plugins/Experimental/TypedElementsDataStorage/Source/TypedElementsDataStorage/Private/TypedElementDatabaseCompatibility.cpp b/Engine/Plugins/Experimental/TypedElementsDataStorage/Source/TypedElementsDataStorage/Private/TypedElementDatabaseCompatibility.cpp
--- a/Engine/Plugins/Experimental/TypedElementsDataStorage/Source/TypedElementsDataStorage/Private/TypedElementDatabaseCompatibility.cpp
+++ b/Engine/Plugins/Experimental/TypedElementsDataStorage/Source/TypedElementsDataStorage/Private/TypedElementDatabaseCompatibility.cpp
@@ -13,6 +13,24 @@
 #include "MassActorEditorSubsystem.h"
 #include "MassActorSubsystem.h"
 #include "TypedElementDataStorageProfilingMacros.h"
+#include "Memento/TypedElementMementoInterface.h"
+
+namespace Private
+{
+	// Number of frames to keep memento rows before deletion
+	bool GMementosEnabled = false;
+	FAutoConsoleVariableRef CVarMementoEnable(
+		TEXT("teds.mementos.enable"),
+		GMementosEnabled,
+		TEXT("Enable memento system for newly added objects\n"));
+
+	void RegisterRowForMementoization(ITypedElementDataStorageInterface* Storage, TypedElementRowHandle Row)
+	{
+		TypedElementTableHandle UnpopulatedMementoTable = UTypedElementMementoInterface::GetUnpopulatedMementoTable();
+		TypedElementRowHandle Memento = Storage->AddRow(UnpopulatedMementoTable);
+		Storage->AddOrGetColumn<FTypedElementMementoOnDelete>(Row, FTypedElementMementoOnDelete{ .Memento = Memento });
+	}
+}
 
 void UTypedElementDatabaseCompatibility::Initialize(ITypedElementDataStorageInterface* StorageInterface)
 {
@@ -423,6 +441,11 @@ void UTypedElementDatabaseCompatibility::TickPendingActorRegistration(UWorld* Ed
 					
 					// Make sure the new row is tagged for update.
 					Storage->AddColumn<FTypedElementSyncFromWorldTag>(Row);
+					
+					if (Private::GMementosEnabled)
+					{
+						Private::RegisterRowForMementoization(Storage, Row);
+					}
 				});
 		}
 			
@@ -455,6 +478,11 @@ void UTypedElementDatabaseCompatibility::TickPendingUObjectRegistration()
 					Storage->AddOrGetColumn<FTypedElementClassTypeInfoColumn>(Row, FTypedElementClassTypeInfoColumn{ .TypeInfo = Object->GetClass() });
 					// Make sure the new row is tagged for update.
 					Storage->AddColumn<FTypedElementSyncFromWorldTag>(Row);
+
+					if (Private::GMementosEnabled)
+					{
+						Private::RegisterRowForMementoization(Storage, Row);
+					}
 				});
 		}
 
@@ -486,6 +514,11 @@ void UTypedElementDatabaseCompatibility::TickPendingExternalObjectRegistration()
 					Storage->AddOrGetColumn<FTypedElementClassTypeInfoColumn>(Row, FTypedElementScriptStructTypeInfoColumn{ .TypeInfo = Object.TypeInfo });
 					// Make sure the new row is tagged for update.
 					Storage->AddColumn<FTypedElementSyncFromWorldTag>(Row);
+
+					if (Private::GMementosEnabled)
+					{
+						Private::RegisterRowForMementoization(Storage, Row);
+					}
 				});
 		}
 
diff --git a/Engine/Plugins/Experimental/TypedElementsDataStorage/Source/TypedElementsDataStorage/Public/Memento/TypedElementMementoInterface.h b/Engine/Plugins/Experimental/TypedElementsDataStorage/Source/TypedElementsDataStorage/Public/Memento/TypedElementMementoInterface.h
new file mode 100644
--- /dev/null
+++ b/Engine/Plugins/Experimental/TypedElementsDataStorage/Source/TypedElementsDataStorage/Public/Memento/TypedElementMementoInterface.h
@@ -0,0 +1,30 @@
+﻿// Copyright Epic Games, Inc. All Rights Reserved.
+
+#pragma once
+
+#include "Elements/Interfaces/TypedElementDataStorageInterface.h"
+
+#include "TypedElementMementoInterface.generated.h"
+
+/**
+ * Add to a row to opt into behaviour to populate a memento
+ * when the row is deleted
+ */
+USTRUCT()
+struct FTypedElementMementoOnDelete : public FTypedElementDataStorageColumn
+{
+	// The memento row populated when the row owning this column is deleted
+	GENERATED_BODY()
+	TypedElementRowHandle Memento;
+};
+
+UCLASS()
+class UTypedElementMementoInterface : public UObject
+{
+	GENERATED_BODY()
+public:
+	/**
+	 * Gets the table type for creating an unpopulated memento row
+	 */
+	static TypedElementTableHandle GetUnpopulatedMementoTable();
+};
\ No newline at end of file
diff --git a/Engine/Plugins/Experimental/TypedElementsDataStorage/Source/TypedElementsDataStorage/Public/Memento/TypedElementMementoTranslators.h b/Engine/Plugins/Experimental/TypedElementsDataStorage/Source/TypedElementsDataStorage/Public/Memento/TypedElementMementoTranslators.h
new file mode 100644
--- /dev/null
+++ b/Engine/Plugins/Experimental/TypedElementsDataStorage/Source/TypedElementsDataStorage/Public/Memento/TypedElementMementoTranslators.h
@@ -0,0 +1,91 @@
+﻿// Copyright Epic Games, Inc. All Rights Reserved.
+
+#pragma once
+
+#include "TypedElementMementoTranslators.generated.h"
+
+/**
+ * A MementoTranslator is used to convert columns of data to mementos and back again
+ * Mementos are special types of columns that hold a subset of column data for dehydrating
+ * a row.  This can be useful for reinstancing, undo/redo, LODs, etc.
+ */
+UCLASS(Abstract)
+class UTypedElementMementoTranslatorBase : public UObject
+{
+	GENERATED_BODY()
+public:
+	/**
+	 * The type information of the column that this translator acts on
+	 */
+	virtual const UScriptStruct* GetColumnType() const
+	{
+		LowLevelFatalError(TEXT("Pure virtual not implemented (UTypedElementColumnMementoProcessorBase)"));
+		return nullptr;
+	}
+
+	/**
+	 * Type information for the memento
+	 */
+	virtual const UScriptStruct* GetMementoType() const
+	{
+		LowLevelFatalError(TEXT("Pure virtual not implemented (UTypedElementColumnMementoProcessorBase)"));
+		return nullptr;
+	}
+	
+	/**
+	 * Implements the logic to convert a column to a memento
+	 */
+	virtual void TranslateColumnToMemento(const void* TypeErasedColumn, void* TypeErasedMemento) const PURE_VIRTUAL(UTypedElementColumnMementoProcessorBase, )
+};
+
+/**
+ * Turn-key translator to convert columns to mementos and back
+ * Creates a new struct that has the same UPROPERTIES the ColumnType
+ * To opt-in a column to be mementoized, override the GetColumnType with the column that should be mementoized
+ *
+ * As an example, suppose there was a column defined as:
+ * ```
+ * USTRUCT()
+ * struct FPathColumn : public FTypedElementDataStorageColumn
+ * {
+ *		GENERATED_BODY()
+ *		
+ *		UPROPERTY()
+ *		FString Path;
+ * };
+ * ```
+ * 
+ * To opt into a default memento translator for FPathColumn, define a default translator:
+ * ```
+ * UCLASS()
+ * class UTypedElementPathColumnTranslator : public UTypedElementDefaultMementoTranslator
+ * {
+ *     virtual const UScriptStruct* GetColumnType() const override { return FPathColumn::StaticStruct; }
+ * };
+ * ```
+ * The system will take care of doing the required operations when rows are deleted that
+ * that have both FTypedElementMementoOnDelete and FPathColumn columns
+ */
+UCLASS(Abstract)
+class UTypedElementDefaultMementoTranslator : public UTypedElementMementoTranslatorBase
+{
+	GENERATED_BODY()
+public:
+	/**
+	 * Override this in derived classes to get automatic mementoization of the given column
+	 */
+	virtual const UScriptStruct* GetColumnType() const override { return nullptr; }
+	virtual const UScriptStruct* GetMementoType() const final override;
+
+	virtual void PostInitProperties() final override;
+
+private:
+	virtual void TranslateColumnToMemento(const void* Column, void* Memento) const final override;
+	
+	UPROPERTY()
+	TObjectPtr<const UScriptStruct> MementoType = nullptr;
+
+	TArray<const FProperty*> SourceProperties;
+	TArray<const FProperty*> DestinationProperties;
+};
+
