Enable TypedElementDataStorageTag columns to be added to a row from a IQueryContext
* This is in line with TEDS treating both fragments and tags as columns.  See UTypedElementDatabase::AddColumn

diff --git a/Engine/Plugins/Experimental/TypedElementsDataStorage/Source/TypedElementsDataStorage/Private/Processors/TypedElementProcessorAdaptors.cpp b/Engine/Plugins/Experimental/TypedElementsDataStorage/Source/TypedElementsDataStorage/Private/Processors/TypedElementProcessorAdaptors.cpp
--- a/Engine/Plugins/Experimental/TypedElementsDataStorage/Source/TypedElementsDataStorage/Private/Processors/TypedElementProcessorAdaptors.cpp
+++ b/Engine/Plugins/Experimental/TypedElementsDataStorage/Source/TypedElementsDataStorage/Private/Processors/TypedElementProcessorAdaptors.cpp
@@ -235,6 +235,8 @@ public:
 
 	void* AddColumnUnitialized(TypedElementDataStorage::RowHandle Row, const UScriptStruct* ObjectType) override
 	{
+		checkf(ObjectType->IsChildOf(FMassTag::StaticStruct()) || ObjectType->IsChildOf(FMassFragment::StaticStruct()), TEXT("Column [%s] must be a Mass Fragment or Tag"), *ObjectType->GetName());
+		
 		struct FAddValueColumn
 		{
 			const UScriptStruct* FragmentType;
@@ -272,24 +274,37 @@ public:
 			AddedColumn = ScratchBuffer.Emplace<FAddValueColumnWithDestructor>(ObjectType, FMassEntityHandle::FromNumber(Row), ObjectCopy);
 		}
 
-		Context.Defer().PushCommand<FMassDeferredAddCommand>(
-			[AddedColumn](FMassEntityManager& System)
-			{
-				System.AddFragmentToEntity(AddedColumn->Entity, AddedColumn->FragmentType);
-			});
-		
-		Context.Defer().PushCommand<FMassDeferredSetCommand>(
-			[AddedColumn](FMassEntityManager& System)
-			{
-				FStructView Fragment = System.GetFragmentDataStruct(AddedColumn->Entity, AddedColumn->FragmentType);
-				AddedColumn->FragmentType->CopyScriptStruct(Fragment.GetMemory(), AddedColumn->Object);
-			});
+		if (ObjectType->IsChildOf(FMassFragment::StaticStruct()))
+		{
+			Context.Defer().PushCommand<FMassDeferredAddCommand>(
+				[AddedColumn](FMassEntityManager& System)
+				{
+					System.AddFragmentToEntity(AddedColumn->Entity, AddedColumn->FragmentType);
+				});
+			
+			Context.Defer().PushCommand<FMassDeferredSetCommand>(
+				[AddedColumn](FMassEntityManager& System)
+				{
+					FStructView Fragment = System.GetFragmentDataStruct(AddedColumn->Entity, AddedColumn->FragmentType);
+					AddedColumn->FragmentType->CopyScriptStruct(Fragment.GetMemory(), AddedColumn->Object);
+				});
+		}
+		else if (ObjectType->IsChildOf(FMassTag::StaticStruct()))
+		{
+			Context.Defer().PushCommand<FMassDeferredAddCommand>(
+				[AddedColumn](FMassEntityManager& System)
+				{
+					System.AddTagToEntity(AddedColumn->Entity, AddedColumn->FragmentType);
+				});
+		}
 
 		return ObjectCopy;
 	}
 
 	void* AddColumnUnitialized(TypedElementDataStorage::RowHandle Row, const UScriptStruct* ObjectType, ObjectMoveOperator Mover) override
 	{
+		checkf(ObjectType->IsChildOf(FMassTag::StaticStruct()) || ObjectType->IsChildOf(FMassFragment::StaticStruct()), TEXT("Column [%s] must be a Mass Fragment or Tag"), *ObjectType->GetName());
+		
 		struct FAddMoveableValueColumn
 		{
 			ObjectMoveOperator Mover;
@@ -330,18 +345,29 @@ public:
 			AddedColumn = ScratchBuffer.Emplace<FAddMoveableValueColumnWithDestructor>(Mover, ObjectType, FMassEntityHandle::FromNumber(Row), MovedObject);
 		}
 
-		Context.Defer().PushCommand<FMassDeferredAddCommand>(
-			[AddedColumn](FMassEntityManager& System)
-			{
-				System.AddFragmentToEntity(AddedColumn->Entity, AddedColumn->FragmentType);
-			});
+		if (ObjectType->IsChildOf(FMassFragment::StaticStruct()))
+		{
+			Context.Defer().PushCommand<FMassDeferredAddCommand>(
+				[AddedColumn](FMassEntityManager& System)
+				{
+					System.AddFragmentToEntity(AddedColumn->Entity, AddedColumn->FragmentType);
+				});
 
-		Context.Defer().PushCommand<FMassDeferredSetCommand>(
-			[AddedColumn](FMassEntityManager& System)
-			{
-				FStructView Fragment = System.GetFragmentDataStruct(AddedColumn->Entity, AddedColumn->FragmentType);
-				AddedColumn->Mover(Fragment.GetMemory(), AddedColumn->Object);
-			});
+			Context.Defer().PushCommand<FMassDeferredSetCommand>(
+				[AddedColumn](FMassEntityManager& System)
+				{
+					FStructView Fragment = System.GetFragmentDataStruct(AddedColumn->Entity, AddedColumn->FragmentType);
+					AddedColumn->Mover(Fragment.GetMemory(), AddedColumn->Object);
+				});
+		}
+		else if (ObjectType->IsChildOf(FMassTag::StaticStruct()))
+		{
+			Context.Defer().PushCommand<FMassDeferredAddCommand>(
+				[AddedColumn](FMassEntityManager& System)
+				{
+					System.AddTagToEntity(AddedColumn->Entity, AddedColumn->FragmentType);
+				});
+		}
 
 		return MovedObject;
 	}
