Create temporary OnObjectAdded/PreRemove delegate shim into TEDS compatibility layer
* OnAdded enables ability for external system to to intercept object row creation to add/remove columns based on type information
* Initially used by memento system to inject memento reference onto a row that initiates the mementoization data processing pipeline.  By default, the memento system is disabled.
* PreRemove added for symmetry and not used in this CL.  Expected to be used when memento pipeline fully implemented.
* Expected to not live forever, serve as a way to inject functionality from memento/reinstancing system wile it is developed

This also allows removal of that awkward UTypedElementMementoInterface shim to get the UnpopulatedMementoTable

diff --git a/Engine/Plugins/Experimental/TypedElementsDataStorage/Source/TypedElementsDataStorage/Private/Memento/TypedElementMementoInterface.cpp b/Engine/Plugins/Experimental/TypedElementsDataStorage/Source/TypedElementsDataStorage/Private/Memento/TypedElementMementoInterface.cpp
deleted file mode 100644
diff --git a/Engine/Plugins/Experimental/TypedElementsDataStorage/Source/TypedElementsDataStorage/Private/Memento/TypedElementMementoSystem.cpp b/Engine/Plugins/Experimental/TypedElementsDataStorage/Source/TypedElementsDataStorage/Private/Memento/TypedElementMementoSystem.cpp
--- a/Engine/Plugins/Experimental/TypedElementsDataStorage/Source/TypedElementsDataStorage/Private/Memento/TypedElementMementoSystem.cpp
+++ b/Engine/Plugins/Experimental/TypedElementsDataStorage/Source/TypedElementsDataStorage/Private/Memento/TypedElementMementoSystem.cpp
@@ -9,8 +9,18 @@
 #include "PropertyBag.h"
 #include "TypedElementMementoRowTypes.h"
 
+DECLARE_LOG_CATEGORY_CLASS(LogTypedElementMemento, Log, All)
+
 namespace Private
 {
+
+	// Number of frames to keep memento rows before deletion
+	bool GMementosEnabled = false;
+	FAutoConsoleVariableRef CVarMementoEnable(
+		TEXT("teds.mementos.enable"),
+		GMementosEnabled,
+		TEXT("Enable memento system for newly added objects\n"));
+
 	// Number of frames to keep memento rows before deletion
 	int32 GMementoKeepFrames = 120;
 	FAutoConsoleVariableRef CVarMementoKeepFrame(
@@ -30,10 +40,56 @@ void UTypedElementMementoSystemFactory::RegisterTables(ITypedElementDataStorageI
 		MementoRowBaseTable,
 		{FTypedElementMementoPopulated::StaticStruct()},
 		TEXT("MementoRowBasePopulatedTable"));
+}
 
-	// Note: Need to register PopulatedMementTable so that "Add Memento Deletion Policy Data" processor does not get pruned
-	// This is because Mass prunes processors that register queries which don't have a matching archtype
-	// at processor construction time
+void UTypedElementMementoSystemFactory::RegisterRegistrationFilters(
+	ITypedElementDataStorageCompatibilityInterface& DataStorageCompatibility) const
+{
+	const_cast<UTypedElementMementoSystemFactory*>(this)->RegisterWithCompatibilityLayer(DataStorageCompatibility);
+}
+
+void UTypedElementMementoSystemFactory::RegisterWithCompatibilityLayer(ITypedElementDataStorageCompatibilityInterface& DataStorageCompatibility)
+{
+	if (Private::GMementosEnabled)
+	{
+		PRAGMA_DISABLE_DEPRECATION_WARNINGS
+		ObjectAddedDelegateHandle = DataStorageCompatibility.GetOnObjectAddedDelegate().AddUObject(this, &UTypedElementMementoSystemFactory::HandleObjectAddedToCompatibility);
+		PRAGMA_ENABLE_DEPRECATION_WARNINGS
+	}
+	
+	auto HandleCVarEnabledChanged = [this, &DataStorageCompatibility](IConsoleVariable* CVar)
+	{
+		// Toggled Off
+		if (!Private::GMementosEnabled)
+		{
+			if (ObjectAddedDelegateHandle.IsValid())
+			{
+				PRAGMA_DISABLE_DEPRECATION_WARNINGS
+				DataStorageCompatibility.GetOnObjectAddedDelegate().Remove(ObjectAddedDelegateHandle);
+				PRAGMA_ENABLE_DEPRECATION_WARNINGS
+				ObjectAddedDelegateHandle.Reset();
+			}
+		}
+		// Toggled on
+		else if (Private::GMementosEnabled)
+		{
+			if (!ObjectAddedDelegateHandle.IsValid())
+			{
+				PRAGMA_DISABLE_DEPRECATION_WARNINGS
+				ObjectAddedDelegateHandle = DataStorageCompatibility.GetOnObjectAddedDelegate().AddUObject(this, &UTypedElementMementoSystemFactory::HandleObjectAddedToCompatibility);
+				PRAGMA_ENABLE_DEPRECATION_WARNINGS
+			}
+		}
+	};
+	
+	Private::CVarMementoEnable->OnChangedDelegate().AddLambda(HandleCVarEnabledChanged);
+}
+
+void UTypedElementMementoSystemFactory::HandleObjectAddedToCompatibility(ITypedElementDataStorageInterface* Storage, const void* Object, const FTypedElementDatabaseCompatibilityObjectTypeInfo& TypeInfo, TypedElementRowHandle Row)
+{
+	// Register row for mementoization
+	TypedElementRowHandle Memento = Storage->AddRow(MementoRowBaseTable);
+	Storage->AddOrGetColumn<FTypedElementMementoOnDelete>(Row, FTypedElementMementoOnDelete{ .Memento = Memento });
 }
 
 void UTypedElementMementoSystemFactory::RegisterQueries(ITypedElementDataStorageInterface& DataStorage) const
@@ -131,8 +187,6 @@ void UTypedElementMementoSystemFactory::RegisterQueries(ITypedElementDataStorage
 			);
 		check(QueryHandle != TypedElementInvalidQueryHandle);
 	}
-
-	
 }
 
 
diff --git a/Engine/Plugins/Experimental/TypedElementsDataStorage/Source/TypedElementsDataStorage/Private/Memento/TypedElementMementoSystem.h b/Engine/Plugins/Experimental/TypedElementsDataStorage/Source/TypedElementsDataStorage/Private/Memento/TypedElementMementoSystem.h
--- a/Engine/Plugins/Experimental/TypedElementsDataStorage/Source/TypedElementsDataStorage/Private/Memento/TypedElementMementoSystem.h
+++ b/Engine/Plugins/Experimental/TypedElementsDataStorage/Source/TypedElementsDataStorage/Private/Memento/TypedElementMementoSystem.h
@@ -1,4 +1,4 @@
-﻿// Copyright Epic Games, Inc. All Rights Reserved.
+// Copyright Epic Games, Inc. All Rights Reserved.
 
 #pragma once
 
@@ -7,6 +7,9 @@
 
 #include "TypedElementMementoSystem.generated.h"
 
+class ITypedElementDataStorageInterface;
+class IConsoleVariable;
+struct FTypedElementDatabaseCompatibilityObjectTypeInfo;
 
 UCLASS()
 class UTypedElementMementoSystemFactory : public UTypedElementDataStorageFactory
@@ -16,13 +19,12 @@ class UTypedElementMementoSystemFactory : public UTypedElementDataStorageFactory
 public:
 	void RegisterQueries(ITypedElementDataStorageInterface& DataStorage) const override;
 	void RegisterTables(ITypedElementDataStorageInterface& DataStorage) const override;
-
-	TypedElementTableHandle GetUnpopulatedMementoTable() const;
+	void RegisterRegistrationFilters(ITypedElementDataStorageCompatibilityInterface& DataStorageCompatibility) const override;
 private:
+
+	void RegisterWithCompatibilityLayer(ITypedElementDataStorageCompatibilityInterface& DataStorageCompatibility);
+	void HandleObjectAddedToCompatibility(ITypedElementDataStorageInterface* Storage, const void* Object, const FTypedElementDatabaseCompatibilityObjectTypeInfo& TypeInfo, TypedElementRowHandle Row);
+	
 	mutable TypedElementTableHandle MementoRowBaseTable;
+	FDelegateHandle ObjectAddedDelegateHandle;
 };
-
-inline TypedElementTableHandle UTypedElementMementoSystemFactory::GetUnpopulatedMementoTable() const
-{
-	return MementoRowBaseTable;
-}
diff --git a/Engine/Plugins/Experimental/TypedElementsDataStorage/Source/TypedElementsDataStorage/Private/TypedElementDatabaseCompatibility.cpp b/Engine/Plugins/Experimental/TypedElementsDataStorage/Source/TypedElementsDataStorage/Private/TypedElementDatabaseCompatibility.cpp
--- a/Engine/Plugins/Experimental/TypedElementsDataStorage/Source/TypedElementsDataStorage/Private/TypedElementDatabaseCompatibility.cpp
+++ b/Engine/Plugins/Experimental/TypedElementsDataStorage/Source/TypedElementsDataStorage/Private/TypedElementDatabaseCompatibility.cpp
@@ -15,23 +15,6 @@
 #include "TypedElementDataStorageProfilingMacros.h"
 #include "Memento/TypedElementMementoInterface.h"
 
-namespace Private
-{
-	// Number of frames to keep memento rows before deletion
-	bool GMementosEnabled = false;
-	FAutoConsoleVariableRef CVarMementoEnable(
-		TEXT("teds.mementos.enable"),
-		GMementosEnabled,
-		TEXT("Enable memento system for newly added objects\n"));
-
-	void RegisterRowForMementoization(ITypedElementDataStorageInterface* Storage, TypedElementRowHandle Row)
-	{
-		TypedElementTableHandle UnpopulatedMementoTable = UTypedElementMementoInterface::GetUnpopulatedMementoTable();
-		TypedElementRowHandle Memento = Storage->AddRow(UnpopulatedMementoTable);
-		Storage->AddOrGetColumn<FTypedElementMementoOnDelete>(Row, FTypedElementMementoOnDelete{ .Memento = Memento });
-	}
-}
-
 void UTypedElementDatabaseCompatibility::Initialize(ITypedElementDataStorageInterface* StorageInterface)
 {
 	checkf(StorageInterface, TEXT("Typed Element's Database compatibility manager is being initialized with an invalid storage target."));
@@ -150,6 +133,17 @@ void UTypedElementDatabaseCompatibility::RemoveCompatibleObjectExplicit(UObject*
 		TypedElementRowHandle Row;
 		if (ReverseObjectLookup.RemoveAndCopyValue(Object, Row))
 		{
+			if (OnObjectPreDestroy.IsBound() && Storage->IsRowAvailable(Row))
+			{
+				// This will be null if object has not yet been added (ie. is in UObjectsPendingRegistration)
+				// Don't call it as OnAdded hasn't be called yet
+				const FTypedElementClassTypeInfoColumn* TypeInfo = Storage->GetColumn<FTypedElementClassTypeInfoColumn>(Row);
+				if (TypeInfo != nullptr)
+				{
+					OnObjectPreDestroy.Broadcast(Storage, Object, FTypedElementDatabaseCompatibilityObjectTypeInfo(TypeInfo->TypeInfo.Get()), Row);
+				}
+			}
+			
 			Storage->RemoveRow(Row);
 		}
 	}
@@ -161,6 +155,17 @@ void UTypedElementDatabaseCompatibility::RemoveCompatibleObjectExplicit(void* Ob
 	TypedElementRowHandle Row;
 	if (ReverseObjectLookup.RemoveAndCopyValue(Object, Row))
 	{
+		if (OnObjectPreDestroy.IsBound() && Storage->IsRowAvailable(Row))
+		{
+			// This will be null if object has not yet been added (ie. is in ExternalObjectsPendingRegistration)
+			// Don't call it as OnAdded hasn't be called yet
+			const FTypedElementScriptStructTypeInfoColumn* TypeInfo = Storage->GetColumn<FTypedElementScriptStructTypeInfoColumn>(Row);
+			if (TypeInfo != nullptr)
+			{
+				OnObjectPreDestroy.Broadcast(Storage, Object, FTypedElementDatabaseCompatibilityObjectTypeInfo(TypeInfo->TypeInfo.Get()), Row);
+			}
+		}
+		
 		Storage->RemoveRow(Row);
 	}
 }
@@ -183,6 +188,18 @@ void UTypedElementDatabaseCompatibility::RemoveCompatibleObjectExplicit(AActor*
 			auto ActorStore = Storage->GetColumn<FMassActorFragment>(Entity.AsNumber());
 			if (ActorStore && !ActorStore->IsOwnedByMass()) // Only remove actors that were externally created.
 			{
+				TypedElementRowHandle Row = Entity.AsNumber();
+				if (OnObjectPreDestroy.IsBound() && Storage->IsRowAvailable(Entity.AsNumber()))
+				{
+					// This will be null if object has not yet been added (ie. is in ActorsPendingRegistration)
+					// Don't call it as OnAdded hasn't be called yet
+					const FTypedElementClassTypeInfoColumn* TypeInfo = Storage->GetColumn<FTypedElementClassTypeInfoColumn>(Row);
+					if (TypeInfo != nullptr)
+					{
+						OnObjectPreDestroy.Broadcast(Storage, Actor, FTypedElementDatabaseCompatibilityObjectTypeInfo(TypeInfo->TypeInfo.Get()), Row);
+					}
+				}
+				
 				ActorSubsystem->RemoveHandleForActor(Actor);
 				Storage->RemoveRow(Entity.AsNumber());
 			}
@@ -234,6 +251,16 @@ TypedElementRowHandle UTypedElementDatabaseCompatibility::FindRowWithCompatibleO
 	return TypedElementInvalidRowHandle;
 }
 
+FTypedElementDatabaseCompatibility_OnObjectAdded& UTypedElementDatabaseCompatibility::GetOnObjectAddedDelegate()
+{
+	return OnObjectAddedDelegate;
+}
+
+FTypedElementDatabaseCompatibility_OnObjectAdded& UTypedElementDatabaseCompatibility::GetOnObjectPreDestroy()
+{
+	return OnObjectPreDestroy;
+}
+
 void UTypedElementDatabaseCompatibility::Prepare()
 {
 	UMassActorEditorSubsystem* MassActorEditorSubsystem = Storage->GetExternalSystem<UMassActorEditorSubsystem>();
@@ -441,11 +468,8 @@ void UTypedElementDatabaseCompatibility::TickPendingActorRegistration(UWorld* Ed
 					
 					// Make sure the new row is tagged for update.
 					Storage->AddColumn<FTypedElementSyncFromWorldTag>(Row);
-					
-					if (Private::GMementosEnabled)
-					{
-						Private::RegisterRowForMementoization(Storage, Row);
-					}
+
+					OnObjectAddedDelegate.Broadcast(Storage, Actor, FTypedElementDatabaseCompatibilityObjectTypeInfo(Actor->GetClass()) , Row);
 				});
 		}
 			
@@ -479,10 +503,7 @@ void UTypedElementDatabaseCompatibility::TickPendingUObjectRegistration()
 					// Make sure the new row is tagged for update.
 					Storage->AddColumn<FTypedElementSyncFromWorldTag>(Row);
 
-					if (Private::GMementosEnabled)
-					{
-						Private::RegisterRowForMementoization(Storage, Row);
-					}
+					OnObjectAddedDelegate.Broadcast(Storage, Object.Get(), FTypedElementDatabaseCompatibilityObjectTypeInfo(Object->GetClass()), Row);
 				});
 		}
 
@@ -515,10 +536,7 @@ void UTypedElementDatabaseCompatibility::TickPendingExternalObjectRegistration()
 					// Make sure the new row is tagged for update.
 					Storage->AddColumn<FTypedElementSyncFromWorldTag>(Row);
 
-					if (Private::GMementosEnabled)
-					{
-						Private::RegisterRowForMementoization(Storage, Row);
-					}
+					OnObjectAddedDelegate.Broadcast(Storage, Object.Object, FTypedElementDatabaseCompatibilityObjectTypeInfo(Object.TypeInfo.Get()), Row);
 				});
 		}
 
diff --git a/Engine/Plugins/Experimental/TypedElementsDataStorage/Source/TypedElementsDataStorage/Private/TypedElementDatabaseCompatibility.h b/Engine/Plugins/Experimental/TypedElementsDataStorage/Source/TypedElementsDataStorage/Private/TypedElementDatabaseCompatibility.h
--- a/Engine/Plugins/Experimental/TypedElementsDataStorage/Source/TypedElementsDataStorage/Private/TypedElementDatabaseCompatibility.h
+++ b/Engine/Plugins/Experimental/TypedElementsDataStorage/Source/TypedElementsDataStorage/Private/TypedElementDatabaseCompatibility.h
@@ -43,6 +43,9 @@ public:
 	TypedElementRowHandle FindRowWithCompatibleObjectExplicit(const UObject* Object) const override;
 	TypedElementRowHandle FindRowWithCompatibleObjectExplicit(const AActor* Actor) const override;
 	TypedElementRowHandle FindRowWithCompatibleObjectExplicit(const void* Object) const override;
+	
+	FTypedElementDatabaseCompatibility_OnObjectAdded& GetOnObjectAddedDelegate() override;
+	FTypedElementDatabaseCompatibility_OnObjectPreDestroy& GetOnObjectPreDestroy() override;
 
 private:
 	void Prepare();
@@ -91,6 +94,8 @@ private:
 	
 	TArray<ObjectRegistrationFilter> ObjectRegistrationFilters;
 	TArray<ObjectToRowDealiaser> ObjectToRowDialiasers;
+	FTypedElementDatabaseCompatibility_OnObjectAdded OnObjectAddedDelegate;
+	FTypedElementDatabaseCompatibility_OnObjectPreDestroy OnObjectPreDestroy;
 
 	TypedElementTableHandle StandardActorTable{ TypedElementInvalidTableHandle };
 	TypedElementTableHandle StandardActorWithTransformTable{ TypedElementInvalidTableHandle };
diff --git a/Engine/Plugins/Experimental/TypedElementsDataStorage/Source/TypedElementsDataStorage/Public/Memento/TypedElementMementoInterface.h b/Engine/Plugins/Experimental/TypedElementsDataStorage/Source/TypedElementsDataStorage/Public/Memento/TypedElementMementoInterface.h
--- a/Engine/Plugins/Experimental/TypedElementsDataStorage/Source/TypedElementsDataStorage/Public/Memento/TypedElementMementoInterface.h
+++ b/Engine/Plugins/Experimental/TypedElementsDataStorage/Source/TypedElementsDataStorage/Public/Memento/TypedElementMementoInterface.h
@@ -16,15 +16,4 @@ struct FTypedElementMementoOnDelete : public FTypedElementDataStorageColumn
 	// The memento row populated when the row owning this column is deleted
 	GENERATED_BODY()
 	TypedElementRowHandle Memento;
-};
-
-UCLASS()
-class UTypedElementMementoInterface : public UObject
-{
-	GENERATED_BODY()
-public:
-	/**
-	 * Gets the table type for creating an unpopulated memento row
-	 */
-	static TypedElementTableHandle GetUnpopulatedMementoTable();
 };
\ No newline at end of file
diff --git a/Engine/Source/Runtime/TypedElementFramework/Public/Elements/Interfaces/TypedElementDataStorageCompatibilityInterface.h b/Engine/Source/Runtime/TypedElementFramework/Public/Elements/Interfaces/TypedElementDataStorageCompatibilityInterface.h
--- a/Engine/Source/Runtime/TypedElementFramework/Public/Elements/Interfaces/TypedElementDataStorageCompatibilityInterface.h
+++ b/Engine/Source/Runtime/TypedElementFramework/Public/Elements/Interfaces/TypedElementDataStorageCompatibilityInterface.h
@@ -14,6 +14,21 @@
 #include "TypedElementDataStorageCompatibilityInterface.generated.h"
 
 class AActor;
+struct FTypedElementDatabaseCompatibilityObjectTypeInfo;
+
+DECLARE_MULTICAST_DELEGATE_FourParams(
+	FTypedElementDatabaseCompatibility_OnObjectAdded,
+	ITypedElementDataStorageInterface* Storage,
+	const void* Object,
+	const FTypedElementDatabaseCompatibilityObjectTypeInfo& ObjectTypeInfo,
+	TypedElementRowHandle Row);
+
+DECLARE_MULTICAST_DELEGATE_FourParams(
+	FTypedElementDatabaseCompatibility_OnObjectPreDestroy,
+	ITypedElementDataStorageInterface* Storage,
+	const void* Object,
+	const FTypedElementDatabaseCompatibilityObjectTypeInfo& ObjectTypeInfo,
+	TypedElementRowHandle Row);
 
 UINTERFACE(MinimalAPI)
 class UTypedElementDataStorageCompatibilityInterface : public UInterface
@@ -32,7 +47,7 @@ class ITypedElementDataStorageCompatibilityInterface
 public:
 	using ObjectRegistrationFilter = TFunction<bool(const ITypedElementDataStorageCompatibilityInterface&, const UObject*)>;
 	using ObjectToRowDealiaser = TFunction<TypedElementRowHandle(const ITypedElementDataStorageCompatibilityInterface&, const UObject*)>;
-
+	
 	/**
 	 * @section Type-agnostic functions
 	 * These allow compatibility with any type. These do eventually fall back to the explicit versions.
@@ -108,6 +123,18 @@ public:
 	virtual TypedElementRowHandle FindRowWithCompatibleObjectExplicit(const AActor* Actor) const = 0;
 	/** Finds a previously stored FStruct. If not found an invalid row handle will be returned. */
 	virtual TypedElementRowHandle FindRowWithCompatibleObjectExplicit(const void* Object) const = 0;
+	
+	/**
+	 * An extension point called when objects have been added to the compatibility interface
+	 */
+	UE_DEPRECATED(5.4, "This is temporary for adding memento support. Do no rely on it.")
+	virtual FTypedElementDatabaseCompatibility_OnObjectAdded& GetOnObjectAddedDelegate() = 0;
+	
+	/**
+	 * An extension point called when objects have been removed from the compatibility interface
+	 */
+	UE_DEPRECATED(5.4, "This is temporary for adding memento support. Do no rely on it.")
+	virtual FTypedElementDatabaseCompatibility_OnObjectPreDestroy& GetOnObjectPreDestroy() = 0;
 };
 
 template<typename Type> Type* GetRawPointer(const TWeakObjectPtr<Type> Object)	{ return Object.Get(); }
@@ -225,3 +252,35 @@ struct TTypedElementSubsystemTraits final
 		}
 	}
 };
+
+enum class ETypedElementDatabaseCompatibilityObjectType
+{
+	Struct,
+	Class
+};
+
+/**
+ * Objects with type info defined in either UScriptStruct or UClass can be stored into TEDS via the
+ * ITypedElementDataStorageCompatibilityInterface
+ * This is a discriminated union which aids with callbacks made when objects are added
+ */
+struct FTypedElementDatabaseCompatibilityObjectTypeInfo
+{
+	ETypedElementDatabaseCompatibilityObjectType TypeInfoType;
+
+	union
+	{
+		const UScriptStruct* ScriptStruct;
+		const UClass* Class;
+	};
+
+	explicit FTypedElementDatabaseCompatibilityObjectTypeInfo(const UScriptStruct* InScriptStruct)
+		: TypeInfoType(ETypedElementDatabaseCompatibilityObjectType::Struct)
+		, ScriptStruct(InScriptStruct)
+	{}
+
+	explicit FTypedElementDatabaseCompatibilityObjectTypeInfo(const UClass* InClass)
+	: TypeInfoType(ETypedElementDatabaseCompatibilityObjectType::Class)
+	, Class(InClass)
+	{}
+};
