Fixed a TEDS shutdown bug.

TEDS was holding on to Mass query data and as a result it caused Mass to crash on shutdown. This change updates the shutdown to include releasing any connections to the phase manager in Mass.

#nojira

diff --git a/Engine/Plugins/Experimental/TypedElementsDataStorage/Source/TypedElementsDataStorage/Private/Queries/TypedElementExtendedQueryStore.cpp b/Engine/Plugins/Experimental/TypedElementsDataStorage/Source/TypedElementsDataStorage/Private/Queries/TypedElementExtendedQueryStore.cpp
--- a/Engine/Plugins/Experimental/TypedElementsDataStorage/Source/TypedElementsDataStorage/Private/Queries/TypedElementExtendedQueryStore.cpp
+++ b/Engine/Plugins/Experimental/TypedElementsDataStorage/Source/TypedElementsDataStorage/Private/Queries/TypedElementExtendedQueryStore.cpp
@@ -39,36 +39,24 @@ void FTypedElementExtendedQueryStore::UnregisterQuery(Handle Query, FMassProcess
 {
 	if (FTypedElementExtendedQuery* QueryData = Get(Query))
 	{
-		if (QueryData->Processor)
+		UnregisterQueryData(Query, *QueryData, PhaseManager);
+		Queries.Remove(Query);
+	}
+}
+
+void FTypedElementExtendedQueryStore::Clear(FMassProcessingPhaseManager& PhaseManager)
+{
+	TickGroupDescriptions.Empty();
+
+	Queries.ListAliveEntries([this, &PhaseManager](Handle Query, FTypedElementExtendedQuery& QueryData)
 		{
-			if (QueryData->Processor->IsA<UTypedElementQueryProcessorCallbackAdapterProcessorBase>())
-			{
-				PhaseManager.UnregisterDynamicProcessor(*QueryData->Processor);
-			}
-			else if (QueryData->Processor->IsA<UTypedElementQueryObserverCallbackAdapterProcessorBase>())
-			{
-				checkf(false, TEXT("Observer queries can not be unregistered."));
-			}
-			else
+			if (QueryData.Processor && QueryData.Processor->IsA<UTypedElementQueryObserverCallbackAdapterProcessorBase>())
 			{
-				checkf(false, TEXT("Query processor %s is of unsupported type %s."),
-					*QueryData->Description.Callback.Name.ToString(), *QueryData->Processor->GetSparseClassDataStruct()->GetName());
+				// Observers can't be unregistered at this point, so skip these for now.
+				return;
 			}
-		}
-		else if (QueryData->Description.Callback.Type == ITypedElementDataStorageInterface::EQueryCallbackType::PhasePreparation)
-		{
-			UnregisterPreambleQuery(QueryData->Description.Callback.Phase, Query);
-		}
-		else if (QueryData->Description.Callback.Type == ITypedElementDataStorageInterface::EQueryCallbackType::PhaseFinalization)
-		{
-			UnregisterPostambleQuery(QueryData->Description.Callback.Phase, Query);
-		}
-		else
-		{
-			QueryData->NativeQuery.Clear();
-		}
-		Queries.Remove(Query);
-	}
+			UnregisterQueryData(Query, QueryData, PhaseManager);
+		});
 }
 
 void FTypedElementExtendedQueryStore::RegisterTickGroup(FName GroupName, ITypedElementDataStorageInterface::EQueryTickPhase Phase,
@@ -138,7 +126,7 @@ bool FTypedElementExtendedQueryStore::IsAlive(Handle Entry) const
 	return Queries.IsAlive(Entry);
 }
 
-void FTypedElementExtendedQueryStore::ListAliveEntries(const ListAliveEntriesCallback& Callback) const
+void FTypedElementExtendedQueryStore::ListAliveEntries(const ListAliveEntriesConstCallback& Callback) const
 {
 	Queries.ListAliveEntries(Callback);
 }
@@ -301,7 +289,7 @@ void FTypedElementExtendedQueryStore::DebugPrintQueryCallbacks(FOutputDevice& Ou
 {
 	Output.Log(TEXT("The Typed Elements Data Storage has the following query callbacks:"));
 	Queries.ListAliveEntries(
-		[&Output](const FTypedElementExtendedQuery& Query)
+		[&Output](Handle QueryHandle, const FTypedElementExtendedQuery& Query)
 		{
 			if (Query.Processor)
 			{
@@ -691,3 +679,35 @@ void FTypedElementExtendedQueryStore::RunPhasePreOrPostAmbleQueries(FMassEntityM
 		}
 	}
 }
+
+void FTypedElementExtendedQueryStore::UnregisterQueryData(Handle Query, FTypedElementExtendedQuery& QueryData, FMassProcessingPhaseManager& PhaseManager)
+{
+	if (QueryData.Processor)
+	{
+		if (QueryData.Processor->IsA<UTypedElementQueryProcessorCallbackAdapterProcessorBase>())
+		{
+			PhaseManager.UnregisterDynamicProcessor(*QueryData.Processor);
+		}
+		else if (QueryData.Processor->IsA<UTypedElementQueryObserverCallbackAdapterProcessorBase>())
+		{
+			checkf(false, TEXT("Observer queries can not be unregistered."));
+		}
+		else
+		{
+			checkf(false, TEXT("Query processor %s is of unsupported type %s."),
+				*QueryData.Description.Callback.Name.ToString(), *QueryData.Processor->GetSparseClassDataStruct()->GetName());
+		}
+	}
+	else if (QueryData.Description.Callback.Type == ITypedElementDataStorageInterface::EQueryCallbackType::PhasePreparation)
+	{
+		UnregisterPreambleQuery(QueryData.Description.Callback.Phase, Query);
+	}
+	else if (QueryData.Description.Callback.Type == ITypedElementDataStorageInterface::EQueryCallbackType::PhaseFinalization)
+	{
+		UnregisterPostambleQuery(QueryData.Description.Callback.Phase, Query);
+	}
+	else
+	{
+		QueryData.NativeQuery.Clear();
+	}
+}
diff --git a/Engine/Plugins/Experimental/TypedElementsDataStorage/Source/TypedElementsDataStorage/Private/Queries/TypedElementExtendedQueryStore.h b/Engine/Plugins/Experimental/TypedElementsDataStorage/Source/TypedElementsDataStorage/Private/Queries/TypedElementExtendedQueryStore.h
--- a/Engine/Plugins/Experimental/TypedElementsDataStorage/Source/TypedElementsDataStorage/Private/Queries/TypedElementExtendedQueryStore.h
+++ b/Engine/Plugins/Experimental/TypedElementsDataStorage/Source/TypedElementsDataStorage/Private/Queries/TypedElementExtendedQueryStore.h
@@ -30,7 +30,7 @@ private:
 	using QueryStore = TTypedElementHandleStore<FTypedElementExtendedQuery>;
 public:
 	using Handle = QueryStore::Handle;
-	using ListAliveEntriesCallback = QueryStore::ListAliveEntriesCallback;
+	using ListAliveEntriesConstCallback = QueryStore::ListAliveEntriesConstCallback;
 
 	/**
 	 * @section Registration
@@ -45,6 +45,9 @@ public:
 		FMassProcessingPhaseManager& PhaseManager);
 	/** Removes the query at the given handle if still alive and otherwise does nothing. */
 	void UnregisterQuery(Handle Query, FMassProcessingPhaseManager& PhaseManager);
+
+	/** Removes all data in the query store. */
+	void Clear(FMassProcessingPhaseManager& PhaseManager);
 	
 	/** Register the defaults for a tick group. These will be applied on top of any settings provided with a query registration. */
 	void RegisterTickGroup(FName GroupName, ITypedElementDataStorageInterface::EQueryTickPhase Phase,
@@ -78,7 +81,7 @@ public:
 	bool IsAlive(Handle Entry) const;
 
 	/** Calls the provided callback for each query that's available. */
-	void ListAliveEntries(const ListAliveEntriesCallback& Callback) const;
+	void ListAliveEntries(const ListAliveEntriesConstCallback& Callback) const;
 
 	/**
 	 * @section Execution
@@ -145,6 +148,8 @@ private:
 	void RunPhasePreOrPostAmbleQueries(FMassEntityManager& EntityManager, FTypedElementDatabaseScratchBuffer& ScratchBuffer,
 		ITypedElementDataStorageInterface::EQueryTickPhase Phase, float DeltaTime, TArray<Handle>& QueryHandles);
 
+	void UnregisterQueryData(Handle Query, FTypedElementExtendedQuery& QueryData, FMassProcessingPhaseManager& PhaseManager);
+
 	static const ITypedElementDataStorageInterface::FQueryDescription EmptyDescription;
 
 	QueryStore Queries;
diff --git a/Engine/Plugins/Experimental/TypedElementsDataStorage/Source/TypedElementsDataStorage/Private/TypedElementDatabase.cpp b/Engine/Plugins/Experimental/TypedElementsDataStorage/Source/TypedElementsDataStorage/Private/TypedElementDatabase.cpp
--- a/Engine/Plugins/Experimental/TypedElementsDataStorage/Source/TypedElementsDataStorage/Private/TypedElementDatabase.cpp
+++ b/Engine/Plugins/Experimental/TypedElementsDataStorage/Source/TypedElementsDataStorage/Private/TypedElementDatabase.cpp
@@ -743,8 +743,13 @@ void UTypedElementDatabase::FinalizePhase(EQueryTickPhase Phase, float DeltaTime
 
 void UTypedElementDatabase::Reset()
 {
+	if (ActiveEditorPhaseManager)
+	{
+		Queries.Clear(*ActiveEditorPhaseManager.Get());
+	}
 	Tables.Reset();
 	TableNameLookup.Reset();
+	ActiveEditorPhaseManager.Reset();
 	ActiveEditorEntityManager.Reset();
 }
 
diff --git a/Engine/Plugins/Experimental/TypedElementsDataStorage/Source/TypedElementsDataStorage/Private/TypedElementHandleStore.h b/Engine/Plugins/Experimental/TypedElementsDataStorage/Source/TypedElementsDataStorage/Private/TypedElementHandleStore.h
--- a/Engine/Plugins/Experimental/TypedElementsDataStorage/Source/TypedElementsDataStorage/Private/TypedElementHandleStore.h
+++ b/Engine/Plugins/Experimental/TypedElementsDataStorage/Source/TypedElementsDataStorage/Private/TypedElementHandleStore.h
@@ -28,7 +28,8 @@ public:
 		uint64 Handle;
 	};
 
-	using ListAliveEntriesCallback = TFunctionRef<void(const DataType&)>;
+	using ListAliveEntriesConstCallback = TFunctionRef<void(Handle, const DataType&)>;
+	using ListAliveEntriesCallback = TFunctionRef<void(Handle, DataType&)>;
 	
 	template<typename... Args>
 	Handle Emplace(Args... Arguments);
@@ -41,7 +42,9 @@ public:
 	void Remove(Handle Entry);
 
 	bool IsAlive(Handle Entry) const;
-	void ListAliveEntries(const ListAliveEntriesCallback& Callback) const;
+	void ListAliveEntries(const ListAliveEntriesConstCallback& Callback) const;
+	void ListAliveEntries(const ListAliveEntriesConstCallback& Callback);
+	void ListAliveEntries(const ListAliveEntriesCallback& Callback);
 
 	friend bool operator==(Handle Lhs, Handle Rhs) { return Lhs.Handle == Rhs.Handle; }
 	friend bool operator!=(Handle Lhs, Handle Rhs) { return Lhs.Handle != Rhs.Handle; }
@@ -137,7 +140,7 @@ bool TTypedElementHandleStore<DataType, ReservationSize>::IsAlive(Handle Entry)
 }
 
 template<typename DataType, uint32 ReservationSize>
-void TTypedElementHandleStore<DataType, ReservationSize>::ListAliveEntries(const ListAliveEntriesCallback& Callback) const
+void TTypedElementHandleStore<DataType, ReservationSize>::ListAliveEntries(const ListAliveEntriesConstCallback& Callback) const
 {
 	int32 Count = Data.Num();
 	const DataType* EntryIt = Data.GetData();
@@ -147,7 +150,38 @@ void TTypedElementHandleStore<DataType, ReservationSize>::ListAliveEntries(const
 	{
 		if (GenerationIt->bIsAlive)
 		{
-			Callback(*EntryIt);
+			Handle DataHandle;
+			DataHandle.Data.Index = Index;
+			DataHandle.Data.Generation = GenerationIt->Generation;
+			Callback(DataHandle, *EntryIt);
+		}
+
+		++EntryIt;
+		++GenerationIt;
+	}
+}
+
+template<typename DataType, uint32 ReservationSize>
+void TTypedElementHandleStore<DataType, ReservationSize>::ListAliveEntries(const ListAliveEntriesConstCallback& Callback)
+{
+	const_cast<const TTypedElementHandleStore<DataType, ReservationSize>*>(this)->ListAliveEntries(Callback);
+}
+
+template<typename DataType, uint32 ReservationSize>
+void TTypedElementHandleStore<DataType, ReservationSize>::ListAliveEntries(const ListAliveEntriesCallback& Callback)
+{
+	int32 Count = Data.Num();
+	DataType* EntryIt = Data.GetData();
+	FGeneration* GenerationIt = Generations.GetData();
+
+	for (int32 Index = 0; Index < Count; ++Index)
+	{
+		if (GenerationIt->bIsAlive)
+		{
+			Handle DataHandle;
+			DataHandle.Data.Index = Index;
+			DataHandle.Data.Generation = GenerationIt->Generation;
+			Callback(DataHandle, *EntryIt);
 		}
 
 		++EntryIt;
