MVP implementation of reinstancing pipeline in TypedElementDataStorage for the compatibility layer
* This CL adds a new column type that is intended to be appended to memento rows to activation the reinstancing process for that row.  This is the FTypedElementReinstanceTarget
* When a row has been used for reinstancing, it is automatically deleted.
* Previous deletion policy of deleting after N frames has been removed.
* Translator has a new interface function that needs to be implemented - TranslateMementoToColumn. This is invoked for all matching columns on a mement row when a reinstance target has been added.  The processors to perform these operations have also been added.
* Misc housekeeping to rename UTypedElementDefaultMementoTranslator's members to remove source and destination as the new function reverses the roles.
* Added to the example in UTypedElementPackagePathMementoTranslator.  Note the default declaration does not need changes.

diff --git a/Engine/Plugins/Experimental/TypedElementsDataStorage/Source/TypedElementsDataStorage/Private/Memento/TypedElementMementoExamples.h b/Engine/Plugins/Experimental/TypedElementsDataStorage/Source/TypedElementsDataStorage/Private/Memento/TypedElementMementoExamples.h
--- a/Engine/Plugins/Experimental/TypedElementsDataStorage/Source/TypedElementsDataStorage/Private/Memento/TypedElementMementoExamples.h
+++ b/Engine/Plugins/Experimental/TypedElementsDataStorage/Source/TypedElementsDataStorage/Private/Memento/TypedElementMementoExamples.h
@@ -43,13 +43,23 @@ public:
 	virtual const UScriptStruct* GetMementoType() const override { return FTypedElementPackagePathColumnMemento::StaticStruct(); }
 
 private:
-	virtual void TranslateColumnToMemento(const void* TypedErasedColumn, void* TypeErasedMemento) const override
+	virtual void TranslateColumnToMemento(const void* TypeErasedColumn, void* TypeErasedMemento) const override
 	{
 		// Boilerplate to get the types back
-		const FTypedElementPackagePathColumn* Column = static_cast<const FTypedElementPackagePathColumn*>(TypedErasedColumn);
+		const FTypedElementPackagePathColumn* Column = static_cast<const FTypedElementPackagePathColumn*>(TypeErasedColumn);
 		FTypedElementPackagePathColumnMemento* Memento = static_cast<FTypedElementPackagePathColumnMemento*>(TypeErasedMemento);
 
 		// Explicit mapping between the FString 
 		Memento->Path = FName(Column->Path);
 	}
+
+	virtual void TranslateMementoToColumn(const void* TypeErasedMemento, void* TypeErasedColumn) const override
+	{
+		// Boilerplate to get the types back
+		const FTypedElementPackagePathColumnMemento* Memento = static_cast<const FTypedElementPackagePathColumnMemento*>(TypeErasedMemento);
+		FTypedElementPackagePathColumn* Column = static_cast<FTypedElementPackagePathColumn*>(TypeErasedColumn);
+
+		// Explicit mapping between the FString 
+		Column->Path = *Memento->Path.ToString();
+	}
 };
\ No newline at end of file
diff --git a/Engine/Plugins/Experimental/TypedElementsDataStorage/Source/TypedElementsDataStorage/Private/Memento/TypedElementMementoRowTypes.h b/Engine/Plugins/Experimental/TypedElementsDataStorage/Source/TypedElementsDataStorage/Private/Memento/TypedElementMementoRowTypes.h
--- a/Engine/Plugins/Experimental/TypedElementsDataStorage/Source/TypedElementsDataStorage/Private/Memento/TypedElementMementoRowTypes.h
+++ b/Engine/Plugins/Experimental/TypedElementsDataStorage/Source/TypedElementsDataStorage/Private/Memento/TypedElementMementoRowTypes.h
@@ -23,9 +23,17 @@ struct FTypedElementMementoTag : public FTypedElementDataStorageTag
  * Column for a memento row indicating that the row is populated with memento data
  */
 USTRUCT()
-struct FTypedElementMementoPopulated : public FTypedElementDataStorageColumn
+struct FTypedElementMementoPopulated : public FTypedElementDataStorageTag
 {
-	// Frame number that the memento row became populated
 	GENERATED_BODY()
-	uint64 FrameNumber;
-};
\ No newline at end of file
+};
+
+/**
+ * A column added to a memento row which will trigger the reinstantiation process
+ */
+USTRUCT()
+struct FTypedElementReinstanceTarget : public FTypedElementDataStorageColumn
+{
+	GENERATED_BODY()
+	TypedElementRowHandle Target;
+};
diff --git a/Engine/Plugins/Experimental/TypedElementsDataStorage/Source/TypedElementsDataStorage/Private/Memento/TypedElementMementoSystem.cpp b/Engine/Plugins/Experimental/TypedElementsDataStorage/Source/TypedElementsDataStorage/Private/Memento/TypedElementMementoSystem.cpp
--- a/Engine/Plugins/Experimental/TypedElementsDataStorage/Source/TypedElementsDataStorage/Private/Memento/TypedElementMementoSystem.cpp
+++ b/Engine/Plugins/Experimental/TypedElementsDataStorage/Source/TypedElementsDataStorage/Private/Memento/TypedElementMementoSystem.cpp
@@ -6,27 +6,18 @@
 #include "HAL/IConsoleManager.h"
 #include "Memento/TypedElementMementoInterface.h"
 #include "Memento/TypedElementMementoTranslators.h"
-#include "PropertyBag.h"
 #include "TypedElementMementoRowTypes.h"
 
 DECLARE_LOG_CATEGORY_CLASS(LogTypedElementMemento, Log, All)
 
 namespace Private
 {
-
 	// Number of frames to keep memento rows before deletion
 	bool GMementosEnabled = false;
 	FAutoConsoleVariableRef CVarMementoEnable(
 		TEXT("teds.mementos.enable"),
 		GMementosEnabled,
 		TEXT("Enable memento system for newly added objects\n"));
-
-	// Number of frames to keep memento rows before deletion
-	int32 GMementoKeepFrames = 120;
-	FAutoConsoleVariableRef CVarMementoKeepFrame(
-		TEXT("teds.mementos.keepframes"),
-		GMementoKeepFrames,
-		TEXT("Number of frames to keep memento rows before deletion\n"));
 }
 
 void UTypedElementMementoSystemFactory::RegisterTables(ITypedElementDataStorageInterface& DataStorage) const
@@ -54,6 +45,7 @@ void UTypedElementMementoSystemFactory::RegisterWithCompatibilityLayer(ITypedEle
 	{
 		PRAGMA_DISABLE_DEPRECATION_WARNINGS
 		ObjectAddedDelegateHandle = DataStorageCompatibility.GetOnObjectAddedDelegate().AddUObject(this, &UTypedElementMementoSystemFactory::HandleObjectAddedToCompatibility);
+		ObjectRemovedDelegateHandle = DataStorageCompatibility.GetOnObjectPreDestroy().AddUObject(this, &UTypedElementMementoSystemFactory::HandleObjectPreRemoveFromCompatibility);
 		PRAGMA_ENABLE_DEPRECATION_WARNINGS
 	}
 	
@@ -66,6 +58,7 @@ void UTypedElementMementoSystemFactory::RegisterWithCompatibilityLayer(ITypedEle
 			{
 				PRAGMA_DISABLE_DEPRECATION_WARNINGS
 				DataStorageCompatibility.GetOnObjectAddedDelegate().Remove(ObjectAddedDelegateHandle);
+				DataStorageCompatibility.GetOnObjectPreDestroy().Remove(ObjectRemovedDelegateHandle);
 				PRAGMA_ENABLE_DEPRECATION_WARNINGS
 				ObjectAddedDelegateHandle.Reset();
 			}
@@ -77,23 +70,73 @@ void UTypedElementMementoSystemFactory::RegisterWithCompatibilityLayer(ITypedEle
 			{
 				PRAGMA_DISABLE_DEPRECATION_WARNINGS
 				ObjectAddedDelegateHandle = DataStorageCompatibility.GetOnObjectAddedDelegate().AddUObject(this, &UTypedElementMementoSystemFactory::HandleObjectAddedToCompatibility);
+				ObjectRemovedDelegateHandle = DataStorageCompatibility.GetOnObjectPreDestroy().AddUObject(this, &UTypedElementMementoSystemFactory::HandleObjectPreRemoveFromCompatibility);
 				PRAGMA_ENABLE_DEPRECATION_WARNINGS
 			}
 		}
 	};
 	
 	Private::CVarMementoEnable->OnChangedDelegate().AddLambda(HandleCVarEnabledChanged);
+	
+	FCoreUObjectDelegates::OnObjectsReinstanced.AddUObject(this, &UTypedElementMementoSystemFactory::HandleOnObjectsReinstanced);
+}
+
+/**
+ * Object reinstancing happens in the following sequence:
+ * - This HandleOnObjectsReinstanced callback is fired
+ * - The old instance objects are deleted
+ * - The new instance objects are added
+ */
+void UTypedElementMementoSystemFactory::HandleOnObjectsReinstanced(
+	const FCoreUObjectDelegates::FReplacementObjectMap& ObjectReplacementMap)
+{
+	for (FCoreUObjectDelegates::FReplacementObjectMap::TConstIterator Iter = ObjectReplacementMap.CreateConstIterator(); Iter; ++Iter)
+	{
+		const UObject* PreDeleteObject = Iter->Key;
+		const UObject* NewInstanceObject = Iter->Value;
+		
+		const TypedElementRowHandle* MementoRowPtr = MementoizableObjects.Find(PreDeleteObject);
+		if (MementoRowPtr != nullptr)
+		{
+			NewInstanceToMementoMap.Add(NewInstanceObject, *MementoRowPtr);
+			UE_LOG(LogTypedElementMemento, VeryVerbose, TEXT("Instance Replacement: 0x%p -> %llu [%s]"), NewInstanceObject, *MementoRowPtr, *NewInstanceObject->GetClass()->GetName());
+		}
+	}
 }
 
 void UTypedElementMementoSystemFactory::HandleObjectAddedToCompatibility(ITypedElementDataStorageInterface* Storage, const void* Object, const FTypedElementDatabaseCompatibilityObjectTypeInfo& TypeInfo, TypedElementRowHandle Row)
 {
-	// Register row for mementoization
-	TypedElementRowHandle Memento = Storage->AddRow(MementoRowBaseTable);
-	Storage->AddOrGetColumn<FTypedElementMementoOnDelete>(Row, FTypedElementMementoOnDelete{ .Memento = Memento });
+	// Register row unconditionally for mementoization
+	{
+		TypedElementRowHandle Memento = Storage->AddRow(MementoRowBaseTable);
+		Storage->AddOrGetColumn<FTypedElementMementoOnDelete>(Row, FTypedElementMementoOnDelete{ .Memento = Memento });
+		UE_LOG(LogTypedElementMemento, VeryVerbose, TEXT("Object [%s] tagged for mementoization 0x%p -> %llu"), *TypeInfo.GetFName().ToString(), Object, Memento);
+		MementoizableObjects.Add(Object, Memento);
+	}
+
+	// If this object is a reinstantiation of a deleted object, then setup the memento row to with the target
+	// row to trigger reinstantiation
+	TypedElementRowHandle Memento;
+	if (NewInstanceToMementoMap.RemoveAndCopyValue(Object, Memento))
+	{
+		UE_LOG(LogTypedElementMemento, VeryVerbose, TEXT("New instance object [%s] mapped to memento 0x%p -> %llu"), *TypeInfo.GetFName().ToString(), Object, Memento);
+		Storage->AddOrGetColumn<FTypedElementReinstanceTarget>(Memento, FTypedElementReinstanceTarget{.Target = Row});
+	}
+}
+
+void UTypedElementMementoSystemFactory::HandleObjectPreRemoveFromCompatibility(ITypedElementDataStorageInterface* Storage,
+	const void* Object, const FTypedElementDatabaseCompatibilityObjectTypeInfo& TypeInfo, TypedElementRowHandle Row)
+{
+	UE_LOG(LogTypedElementMemento, VeryVerbose, TEXT("Removing object 0x%p"),Object);
+
+	// Remove row unconditionally from mementoization
+	MementoizableObjects.Remove(Object);
 }
 
 void UTypedElementMementoSystemFactory::RegisterQueries(ITypedElementDataStorageInterface& DataStorage) const
 {
+	using DSI = ITypedElementDataStorageInterface;
+	
 	TArray<const UTypedElementMementoTranslatorBase*> MementoTranslators;
 
 	// Discover all MementoTranslators
@@ -142,6 +185,45 @@ void UTypedElementMementoSystemFactory::RegisterQueries(ITypedElementDataStorage
 		check(QueryHandle != TypedElementInvalidQueryHandle);
 	}
 
+	// Setup processors to execute memento translators on to reinstantiate rows
+	for (int32 Index = 0, End = MementoTranslators.Num(); Index < End; ++Index)
+	{
+		const UTypedElementMementoTranslatorBase* MementoTranslator = MementoTranslators[Index];
+		const UScriptStruct* MementoizedColumnType = MementoTranslator->GetColumnType();
+		const UScriptStruct* MementoType = MementoTranslator->GetMementoType();
+
+		TypedElementQueryHandle Subquery = DataStorage.RegisterQuery(
+			Select()
+				.ReadWrite(MementoizedColumnType)
+			.Compile());
+
+		const FName TranslationProcessorName = FName(FString::Printf(TEXT("MementoTranslator: %s -> %s"), *MementoType->GetName(), *MementoizedColumnType->GetName()));
+		const TypedElementQueryHandle QueryHandle = DataStorage.RegisterQuery(
+			Select(
+				TranslationProcessorName,
+				FProcessor(DSI::EQueryTickPhase::PostPhysics, DataStorage.GetQueryTickGroupName(DSI::EQueryTickGroups::Default)),
+				[MementoTranslator](TypedElementDataStorage::IQueryContext& Context, TypedElementRowHandle Row, const FTypedElementReinstanceTarget& ReinstanceTarget)
+				{
+					const UScriptStruct* MementoType = MementoTranslator->GetMementoType();
+					
+					const void* Memento = Context.GetColumn(MementoType);
+					Context.RunSubquery(0, ReinstanceTarget.Target, 
+						[MementoTranslator, Memento](const DSI::FQueryDescription&, DSI::ISubqueryContext& SubQueryContext)
+						{
+							const UScriptStruct* MementoizedColumnType = MementoTranslator->GetColumnType();
+							void* MementoizedColumn = SubQueryContext.GetMutableColumn(MementoizedColumnType);
+
+							MementoTranslator->TranslateMementoToColumn(Memento, MementoizedColumn);
+						});
+				})
+				.ReadOnly(MementoType)
+				.Where()
+					.All<FTypedElementMementoTag>()
+				.DependsOn().SubQuery(Subquery)
+				.Compile());
+		check(QueryHandle != TypedElementInvalidQueryHandle);
+	}
+
 	// Setup a default policy of deleting populated mementos after a set number of frames
 	// The first use-case of mementos are to support reinstancing which occurs over a single frame
 	// Therefore it should be safe to remove populated mementos after a couple of frames since the data
@@ -150,39 +232,29 @@ void UTypedElementMementoSystemFactory::RegisterQueries(ITypedElementDataStorage
 	{
 		const TypedElementQueryHandle QueryHandle = DataStorage.RegisterQuery(
 			Select(
-				TEXT("Add Memento Deletion Policy Data"),
+				TEXT("Add Memento Populated Tag"),
 				FObserver::OnRemove<FTypedElementMementoOnDelete>(), // When a memento is populated
 				[](TypedElementDataStorage::IQueryContext& Context, TypedElementRowHandle Row, const FTypedElementMementoOnDelete& MementoRow)
 				{
-					Context.AddColumn<FTypedElementMementoPopulated>(
-						MementoRow.Memento,
-						FTypedElementMementoPopulated
-						{
-							.FrameNumber = GFrameCounter
-						});
+					Context.AddColumn<FTypedElementMementoPopulated>(MementoRow.Memento, FTypedElementMementoPopulated{});
 				})
 				.Compile());
 		check(QueryHandle != TypedElementInvalidQueryHandle);
 	}
 
-	using DSI = ITypedElementDataStorageInterface;
-
 	/**
 	 * A processor which deletes mementos that were populated more than N frames ago
 	 */
 	{
 		const TypedElementQueryHandle QueryHandle = DataStorage.RegisterQuery(
 			Select(
-			TEXT("Delete populated mementos older than N frames"),
+			TEXT("Delete populated mementos"),
 			FProcessor(DSI::EQueryTickPhase::FrameEnd, DataStorage.GetQueryTickGroupName(DSI::EQueryTickGroups::Default)),
-				[](TypedElementDataStorage::IQueryContext& Context, TypedElementRowHandle Row, const FTypedElementMementoPopulated& PopulateFrameNumber)
+				[](TypedElementDataStorage::IQueryContext& Context, TypedElementRowHandle Row)
 				{
-					if (PopulateFrameNumber.FrameNumber + Private::GMementoKeepFrames < GFrameCounter)
-					{
-						Context.RemoveRow(Row);
-					}
+					Context.RemoveRow(Row);
 				})
-				.Where().All<FTypedElementMementoTag>()
+				.Where().All<FTypedElementMementoTag, FTypedElementMementoPopulated>()
 				.Compile()
 			);
 		check(QueryHandle != TypedElementInvalidQueryHandle);
diff --git a/Engine/Plugins/Experimental/TypedElementsDataStorage/Source/TypedElementsDataStorage/Private/Memento/TypedElementMementoSystem.h b/Engine/Plugins/Experimental/TypedElementsDataStorage/Source/TypedElementsDataStorage/Private/Memento/TypedElementMementoSystem.h
--- a/Engine/Plugins/Experimental/TypedElementsDataStorage/Source/TypedElementsDataStorage/Private/Memento/TypedElementMementoSystem.h
+++ b/Engine/Plugins/Experimental/TypedElementsDataStorage/Source/TypedElementsDataStorage/Private/Memento/TypedElementMementoSystem.h
@@ -1,12 +1,15 @@
-// Copyright Epic Games, Inc. All Rights Reserved.
+﻿// Copyright Epic Games, Inc. All Rights Reserved.
 
 #pragma once
 
 #include "Elements/Interfaces/TypedElementDataStorageFactory.h"
 #include "Elements/Columns/TypedElementSelectionColumns.h"
+#include "UObject/UObjectGlobals.h"
 
 #include "TypedElementMementoSystem.generated.h"
 
+
+
 class ITypedElementDataStorageInterface;
 class IConsoleVariable;
 struct FTypedElementDatabaseCompatibilityObjectTypeInfo;
@@ -24,7 +27,16 @@ private:
 
 	void RegisterWithCompatibilityLayer(ITypedElementDataStorageCompatibilityInterface& DataStorageCompatibility);
 	void HandleObjectAddedToCompatibility(ITypedElementDataStorageInterface* Storage, const void* Object, const FTypedElementDatabaseCompatibilityObjectTypeInfo& TypeInfo, TypedElementRowHandle Row);
+	void HandleObjectPreRemoveFromCompatibility(ITypedElementDataStorageInterface* Storage, const void* Object, const FTypedElementDatabaseCompatibilityObjectTypeInfo& TypeInfo, TypedElementRowHandle Row);
+	void HandleOnObjectsReinstanced(const FCoreUObjectDelegates::FReplacementObjectMap& ObjectReplacementMap);
 	
 	mutable TypedElementTableHandle MementoRowBaseTable;
 	FDelegateHandle ObjectAddedDelegateHandle;
+	FDelegateHandle ObjectRemovedDelegateHandle;
+
+	// Maps existing objects to their Memento
+	TMap<const void*, TypedElementRowHandle> MementoizableObjects;
+
+	// Maps objects that are about to be created to the memento that will be used to reinstance them
+	TMap<const void*, TypedElementRowHandle> NewInstanceToMementoMap;
 };
diff --git a/Engine/Plugins/Experimental/TypedElementsDataStorage/Source/TypedElementsDataStorage/Private/Memento/TypedElementMementoTranslators.cpp b/Engine/Plugins/Experimental/TypedElementsDataStorage/Source/TypedElementsDataStorage/Private/Memento/TypedElementMementoTranslators.cpp
--- a/Engine/Plugins/Experimental/TypedElementsDataStorage/Source/TypedElementsDataStorage/Private/Memento/TypedElementMementoTranslators.cpp
+++ b/Engine/Plugins/Experimental/TypedElementsDataStorage/Source/TypedElementsDataStorage/Private/Memento/TypedElementMementoTranslators.cpp
@@ -45,25 +45,23 @@ void UTypedElementDefaultMementoTranslator::PostInitProperties()
 		const FProperty* DestinationProperty = PropertyBag->FindPropertyByName(SourceProperty->GetFName());
 		if (DestinationProperty && DestinationProperty->SameType(SourceProperty))
 		{
-			SourceProperties.Add(SourceProperty);
-			DestinationProperties.Add(DestinationProperty);
+			MementoizedColumnProperties.Add(SourceProperty);
+			MementoProperties.Add(DestinationProperty);
 		}
 	}
 }
 
-void UTypedElementDefaultMementoTranslator::TranslateColumnToMemento(const void* Column, void* Memento) const
+void UTypedElementDefaultMementoTranslator::TranslateColumnToMemento(const void* TypeErasedColumn, void* TypeErasedMemento) const
 {
-	const UScriptStruct* ColumnType = GetColumnType();
+	const std::byte* BaseAddressColumn = static_cast<const std::byte*>(TypeErasedColumn);
+	std::byte* BaseAddressMemento = static_cast<std::byte*>(TypeErasedMemento);
 
-	const std::byte* BaseAddressColumn = static_cast<const std::byte*>(Column);
-	std::byte* BaseAddressMemento = static_cast<std::byte*>(Memento);
+	check(MementoizedColumnProperties.Num() == MementoProperties.Num());
 
-	check(SourceProperties.Num() == DestinationProperties.Num());
-
-	for (int32 PropertyIndex = 0, PropertyIndexEnd = SourceProperties.Num(); PropertyIndex < PropertyIndexEnd; ++PropertyIndex)
+	for (int32 PropertyIndex = 0, PropertyIndexEnd = MementoizedColumnProperties.Num(); PropertyIndex < PropertyIndexEnd; ++PropertyIndex)
 	{
-		const FProperty* SourceProperty = SourceProperties[PropertyIndex];
-		const FProperty* DestinationProperty = DestinationProperties[PropertyIndex];
+		const FProperty* SourceProperty = MementoizedColumnProperties[PropertyIndex];
+		const FProperty* DestinationProperty = MementoProperties[PropertyIndex];
 		void* DestinationValueAddress = BaseAddressMemento + DestinationProperty->GetOffset_ForInternal();
 		const void* SourceValueAddress = BaseAddressColumn + SourceProperty->GetOffset_ForInternal();
 		SourceProperty->CopyCompleteValue(
@@ -71,3 +69,23 @@ void UTypedElementDefaultMementoTranslator::TranslateColumnToMemento(const void*
 			SourceValueAddress);
 	}	
 }
+
+void UTypedElementDefaultMementoTranslator::TranslateMementoToColumn(const void* TypeErasedMemento,
+	void* TypeErasedColumn) const
+{
+	const std::byte* BaseAddressMemento = static_cast<const std::byte*>(TypeErasedMemento);
+	std::byte* BaseAddressColumn = static_cast<std::byte*>(TypeErasedColumn);
+
+	check(MementoizedColumnProperties.Num() == MementoProperties.Num());
+
+	for (int32 PropertyIndex = 0, PropertyIndexEnd = MementoizedColumnProperties.Num(); PropertyIndex < PropertyIndexEnd; ++PropertyIndex)
+	{
+		const FProperty* SourceProperty = MementoProperties[PropertyIndex];
+		const FProperty* DestinationProperty = MementoizedColumnProperties[PropertyIndex];
+		void* DestinationValueAddress = BaseAddressColumn + DestinationProperty->GetOffset_ForInternal();
+		const void* SourceValueAddress = BaseAddressMemento + SourceProperty->GetOffset_ForInternal();
+		SourceProperty->CopyCompleteValue(
+			DestinationValueAddress,
+			SourceValueAddress);
+	}	
+}
diff --git a/Engine/Plugins/Experimental/TypedElementsDataStorage/Source/TypedElementsDataStorage/Public/Memento/TypedElementMementoTranslators.h b/Engine/Plugins/Experimental/TypedElementsDataStorage/Source/TypedElementsDataStorage/Public/Memento/TypedElementMementoTranslators.h
--- a/Engine/Plugins/Experimental/TypedElementsDataStorage/Source/TypedElementsDataStorage/Public/Memento/TypedElementMementoTranslators.h
+++ b/Engine/Plugins/Experimental/TypedElementsDataStorage/Source/TypedElementsDataStorage/Public/Memento/TypedElementMementoTranslators.h
@@ -36,6 +36,10 @@ public:
 	 * Implements the logic to convert a column to a memento
 	 */
 	virtual void TranslateColumnToMemento(const void* TypeErasedColumn, void* TypeErasedMemento) const PURE_VIRTUAL(UTypedElementColumnMementoProcessorBase, )
+	/**
+	 * Implements the logic to convert a memento to a column
+	 */
+	virtual void TranslateMementoToColumn(const void* TypeErasedMemento, void* TypeErasedColumn) const PURE_VIRTUAL(UTypedElementColumnMementoProcessorBase, )
 };
 
 /**
@@ -78,14 +82,15 @@ public:
 	virtual const UScriptStruct* GetMementoType() const final override;
 
 	virtual void PostInitProperties() final override;
+	
+	virtual void TranslateColumnToMemento(const void* Column, void* Memento) const final override;
+	virtual void TranslateMementoToColumn(const void* TypeErasedMemento, void* TypeErasedColumn) const final override;
 
 private:
-	virtual void TranslateColumnToMemento(const void* Column, void* Memento) const final override;
-	
 	UPROPERTY()
 	TObjectPtr<const UScriptStruct> MementoType = nullptr;
 
-	TArray<const FProperty*> SourceProperties;
-	TArray<const FProperty*> DestinationProperties;
+	TArray<const FProperty*> MementoizedColumnProperties;
+	TArray<const FProperty*> MementoProperties;
 };
 
diff --git a/Engine/Source/Runtime/TypedElementFramework/Public/Elements/Interfaces/TypedElementDataStorageCompatibilityInterface.h b/Engine/Source/Runtime/TypedElementFramework/Public/Elements/Interfaces/TypedElementDataStorageCompatibilityInterface.h
--- a/Engine/Source/Runtime/TypedElementFramework/Public/Elements/Interfaces/TypedElementDataStorageCompatibilityInterface.h
+++ b/Engine/Source/Runtime/TypedElementFramework/Public/Elements/Interfaces/TypedElementDataStorageCompatibilityInterface.h
@@ -283,4 +283,17 @@ struct FTypedElementDatabaseCompatibilityObjectTypeInfo
 	: TypeInfoType(ETypedElementDatabaseCompatibilityObjectType::Class)
 	, Class(InClass)
 	{}
+
+	FName GetFName() const;
 };
+
+
+inline FName FTypedElementDatabaseCompatibilityObjectTypeInfo::GetFName() const
+{
+	switch(TypeInfoType)
+	{
+	case ETypedElementDatabaseCompatibilityObjectType::Struct: return ScriptStruct->GetFName();
+	case ETypedElementDatabaseCompatibilityObjectType::Class: return Class->GetFName();
+	default: return FName();
+	}
+}
