Added ability to associate metadata with columns in a TEDS query.

diff --git a/Engine/Plugins/Experimental/TypedElementsDataStorage/Source/TypedElementsDataStorage/Private/TypedElementDatabaseUI.cpp b/Engine/Plugins/Experimental/TypedElementsDataStorage/Source/TypedElementsDataStorage/Private/TypedElementDatabaseUI.cpp
--- a/Engine/Plugins/Experimental/TypedElementsDataStorage/Source/TypedElementsDataStorage/Private/TypedElementDatabaseUI.cpp
+++ b/Engine/Plugins/Experimental/TypedElementsDataStorage/Source/TypedElementsDataStorage/Private/TypedElementDatabaseUI.cpp
@@ -239,7 +239,7 @@ bool UTypedElementDatabaseUi::RegisterWidgetFactory(FName Purpose, TUniquePtr<FT
 }
 
 void UTypedElementDatabaseUi::CreateWidgetConstructors(FName Purpose,
-	TConstArrayView<TypedElement::ColumnUtils::Argument> Arguments, const WidgetConstructorCallback& Callback)
+	TypedElementDataStorage::FMetaDataView Arguments, const WidgetConstructorCallback& Callback)
 {
 	if (FPurposeInfo* PurposeInfo = WidgetPurposes.Find(Purpose))
 	{
@@ -254,7 +254,7 @@ void UTypedElementDatabaseUi::CreateWidgetConstructors(FName Purpose,
 }
 
 void UTypedElementDatabaseUi::CreateWidgetConstructors(FName Purpose, EMatchApproach MatchApproach, 
-	TArray<TWeakObjectPtr<const UScriptStruct>>& Columns, TConstArrayView<TypedElement::ColumnUtils::Argument> Arguments, 
+	TArray<TWeakObjectPtr<const UScriptStruct>>& Columns, TypedElementDataStorage::FMetaDataView Arguments,
 	const WidgetConstructorCallback& Callback)
 {
 	if (FPurposeInfo* PurposeInfo = WidgetPurposes.Find(Purpose))
@@ -316,7 +316,7 @@ void UTypedElementDatabaseUi::CreateWidgetConstructors(FName Purpose, EMatchAppr
 	}
 }
 
-void UTypedElementDatabaseUi::ConstructWidgets(FName Purpose, TConstArrayView<TypedElement::ColumnUtils::Argument> Arguments,
+void UTypedElementDatabaseUi::ConstructWidgets(FName Purpose, TypedElementDataStorage::FMetaDataView Arguments,
 	const WidgetCreatedCallback& ConstructionCallback)
 {
 	if (FPurposeInfo* PurposeInfo = WidgetPurposes.Find(Purpose))
@@ -340,7 +340,7 @@ void UTypedElementDatabaseUi::ConstructWidgets(FName Purpose, TConstArrayView<Ty
 
 bool UTypedElementDatabaseUi::CreateSingleWidgetConstructor(
 	const FWidgetFactory::ConstructorType& Constructor,
-	TConstArrayView<TypedElement::ColumnUtils::Argument> Arguments, 
+	TypedElementDataStorage::FMetaDataView Arguments,
 	TConstArrayView<TWeakObjectPtr<const UScriptStruct>> MatchedColumnTypes, 
 	const WidgetConstructorCallback& Callback)
 {
@@ -376,7 +376,7 @@ bool UTypedElementDatabaseUi::CreateSingleWidgetConstructor(
 
 void UTypedElementDatabaseUi::CreateWidgetInstanceFromDescription(
 	const UScriptStruct* Target,
-	TConstArrayView<TypedElement::ColumnUtils::Argument> Arguments,
+	TypedElementDataStorage::FMetaDataView Arguments,
 	const WidgetCreatedCallback& ConstructionCallback)
 {
 	FTypedElementWidgetConstructor* Constructor = reinterpret_cast<FTypedElementWidgetConstructor*>(
@@ -395,7 +395,7 @@ void UTypedElementDatabaseUi::CreateWidgetInstanceFromDescription(
 
 void UTypedElementDatabaseUi::CreateWidgetInstanceFromInstance(
 	FTypedElementWidgetConstructor* SourceConstructor,
-	TConstArrayView<TypedElement::ColumnUtils::Argument> Arguments,
+	TypedElementDataStorage::FMetaDataView Arguments,
 	const WidgetCreatedCallback& ConstructionCallback)
 {
 	if (SourceConstructor->CanBeReused())
@@ -424,7 +424,7 @@ void UTypedElementDatabaseUi::CreateWidgetInstanceFromInstance(
 
 void UTypedElementDatabaseUi::CreateWidgetInstance(
 	FTypedElementWidgetConstructor& Constructor, 
-	TConstArrayView<TypedElement::ColumnUtils::Argument> Arguments, 
+	TypedElementDataStorage::FMetaDataView Arguments,
 	const WidgetCreatedCallback& ConstructionCallback)
 {
 	TypedElementRowHandle Row = Storage->AddRow(WidgetTable);
@@ -441,7 +441,7 @@ void UTypedElementDatabaseUi::CreateWidgetInstance(
 }
 
 TSharedPtr<SWidget> UTypedElementDatabaseUi::ConstructWidget(TypedElementRowHandle Row, FTypedElementWidgetConstructor& Constructor,
-	TConstArrayView<TypedElement::ColumnUtils::Argument> Arguments)
+	TypedElementDataStorage::FMetaDataView Arguments)
 {
 	if (Constructor.CanBeReused())
 	{
@@ -517,7 +517,7 @@ bool UTypedElementDatabaseUi::PrepareColumnsList(TArray<TWeakObjectPtr<const USc
 }
 
 void UTypedElementDatabaseUi::CreateWidgetConstructors_LongestMatch(const TArray<FWidgetFactory>& WidgetFactories,
-	TArray<TWeakObjectPtr<const UScriptStruct>>& Columns, TConstArrayView<TypedElement::ColumnUtils::Argument> Arguments,
+	TArray<TWeakObjectPtr<const UScriptStruct>>& Columns, TypedElementDataStorage::FMetaDataView Arguments,
 	const WidgetConstructorCallback& Callback)
 {
 	using ColumnView = TConstArrayView<TWeakObjectPtr<const UScriptStruct>>;
@@ -571,7 +571,7 @@ void UTypedElementDatabaseUi::CreateWidgetConstructors_LongestMatch(const TArray
 }
 
 void UTypedElementDatabaseUi::CreateWidgetConstructors_ExactMatch(const TArray<FWidgetFactory>& WidgetFactories,
-	TArray<TWeakObjectPtr<const UScriptStruct>>& Columns, TConstArrayView<TypedElement::ColumnUtils::Argument> Arguments,
+	TArray<TWeakObjectPtr<const UScriptStruct>>& Columns, TypedElementDataStorage::FMetaDataView Arguments,
 	const WidgetConstructorCallback& Callback)
 {
 	int32 ColumnCount = Columns.Num();
@@ -605,7 +605,7 @@ void UTypedElementDatabaseUi::CreateWidgetConstructors_ExactMatch(const TArray<F
 }
 
 void UTypedElementDatabaseUi::CreateWidgetConstructors_SingleMatch(const TArray<FWidgetFactory>& WidgetFactories,
-	TArray<TWeakObjectPtr<const UScriptStruct>>& Columns, TConstArrayView<TypedElement::ColumnUtils::Argument> Arguments,
+	TArray<TWeakObjectPtr<const UScriptStruct>>& Columns, TypedElementDataStorage::FMetaDataView Arguments,
 	const WidgetConstructorCallback& Callback)
 {
 	auto FactoryIt = WidgetFactories.rbegin();
diff --git a/Engine/Plugins/Experimental/TypedElementsDataStorage/Source/TypedElementsDataStorage/Private/TypedElementDatabaseUI.h b/Engine/Plugins/Experimental/TypedElementsDataStorage/Source/TypedElementsDataStorage/Private/TypedElementDatabaseUI.h
--- a/Engine/Plugins/Experimental/TypedElementsDataStorage/Source/TypedElementsDataStorage/Private/TypedElementDatabaseUI.h
+++ b/Engine/Plugins/Experimental/TypedElementsDataStorage/Source/TypedElementsDataStorage/Private/TypedElementDatabaseUI.h
@@ -40,14 +40,14 @@ public:
 		TArray<TWeakObjectPtr<const UScriptStruct>> Columns) override;
 
 	void CreateWidgetConstructors(FName Purpose,
-		TConstArrayView<TypedElement::ColumnUtils::Argument> Arguments, const WidgetConstructorCallback& Callback) override;
+		TypedElementDataStorage::FMetaDataView Arguments, const WidgetConstructorCallback& Callback) override;
 	void CreateWidgetConstructors(FName Purpose, EMatchApproach MatchApproach, TArray<TWeakObjectPtr<const UScriptStruct>>& Columns,
-		TConstArrayView<TypedElement::ColumnUtils::Argument> Arguments, const WidgetConstructorCallback& Callback) override;
+		TypedElementDataStorage::FMetaDataView Arguments, const WidgetConstructorCallback& Callback) override;
 
-	void ConstructWidgets(FName Purpose, TConstArrayView<TypedElement::ColumnUtils::Argument> Arguments,
+	void ConstructWidgets(FName Purpose, TypedElementDataStorage::FMetaDataView Arguments,
 		const WidgetCreatedCallback& ConstructionCallback) override;
 	TSharedPtr<SWidget> ConstructWidget(TypedElementRowHandle Row, FTypedElementWidgetConstructor& Constructor,
-		TConstArrayView<TypedElement::ColumnUtils::Argument> Arguments) override;
+		TypedElementDataStorage::FMetaDataView Arguments) override;
 
 	void ListWidgetPurposes(const WidgetPurposeCallback& Callback) const override;
 
@@ -78,23 +78,23 @@ private:
 
 	bool CreateSingleWidgetConstructor(
 		const FWidgetFactory::ConstructorType& Constructor,
-		TConstArrayView<TypedElement::ColumnUtils::Argument> Arguments, 
+		TypedElementDataStorage::FMetaDataView Arguments,
 		TConstArrayView<TWeakObjectPtr<const UScriptStruct>> MatchedColumnTypes,
 		const WidgetConstructorCallback& Callback);
 
 	void CreateWidgetInstanceFromDescription(
 		const UScriptStruct* Target,
-		TConstArrayView<TypedElement::ColumnUtils::Argument> Arguments,
+		TypedElementDataStorage::FMetaDataView Arguments,
 		const WidgetCreatedCallback& ConstructionCallback);
 
 	void CreateWidgetInstanceFromInstance(
 		FTypedElementWidgetConstructor* SourceConstructor,
-		TConstArrayView<TypedElement::ColumnUtils::Argument> Arguments,
+		TypedElementDataStorage::FMetaDataView Arguments,
 		const WidgetCreatedCallback& ConstructionCallback);
 
 	void CreateWidgetInstance(
 		FTypedElementWidgetConstructor& Constructor,
-		TConstArrayView<TypedElement::ColumnUtils::Argument> Arguments,
+		TypedElementDataStorage::FMetaDataView Arguments,
 		const WidgetCreatedCallback& ConstructionCallback);
 
 	static bool PrepareColumnsList(TArray<TWeakObjectPtr<const UScriptStruct>>& Columns);
@@ -102,17 +102,17 @@ private:
 	void CreateWidgetConstructors_LongestMatch(
 		const TArray<FWidgetFactory>& WidgetFactories, 
 		TArray<TWeakObjectPtr<const UScriptStruct>>& Columns, 
-		TConstArrayView<TypedElement::ColumnUtils::Argument> Arguments, 
+		TypedElementDataStorage::FMetaDataView Arguments,
 		const WidgetConstructorCallback& Callback);
 	void CreateWidgetConstructors_ExactMatch(
 		const TArray<FWidgetFactory>& WidgetFactories,
 		TArray<TWeakObjectPtr<const UScriptStruct>>& Columns,
-		TConstArrayView<TypedElement::ColumnUtils::Argument> Arguments,
+		TypedElementDataStorage::FMetaDataView Arguments,
 		const WidgetConstructorCallback& Callback);
 	void CreateWidgetConstructors_SingleMatch(
 		const TArray<FWidgetFactory>& WidgetFactories,
 		TArray<TWeakObjectPtr<const UScriptStruct>>& Columns,
-		TConstArrayView<TypedElement::ColumnUtils::Argument> Arguments,
+		TypedElementDataStorage::FMetaDataView Arguments,
 		const WidgetConstructorCallback& Callback);
 
 	TypedElementTableHandle WidgetTable{ TypedElementInvalidTableHandle };
diff --git a/Engine/Plugins/Experimental/TypedElementsDataStorage/Source/TypedElementsDataStorageUI/Private/TypedElementCounterWidget.cpp b/Engine/Plugins/Experimental/TypedElementsDataStorage/Source/TypedElementsDataStorageUI/Private/TypedElementCounterWidget.cpp
--- a/Engine/Plugins/Experimental/TypedElementsDataStorage/Source/TypedElementsDataStorageUI/Private/TypedElementCounterWidget.cpp
+++ b/Engine/Plugins/Experimental/TypedElementsDataStorage/Source/TypedElementsDataStorageUI/Private/TypedElementCounterWidget.cpp
@@ -154,12 +154,7 @@ TConstArrayView<const UScriptStruct*> FTypedElementCounterWidgetConstructor::Get
 	return Columns;
 }
 
-bool FTypedElementCounterWidgetConstructor::CanBeReused() const
-{
-	return true;
-}
-
-TSharedPtr<SWidget> FTypedElementCounterWidgetConstructor::CreateWidget()
+TSharedPtr<SWidget> FTypedElementCounterWidgetConstructor::CreateWidget(TypedElementDataStorage::FMetaDataView Arguments)
 {
 	return SNew(STextBlock)
 		.Text(FText::Format(LabelText, 0))
diff --git a/Engine/Plugins/Experimental/TypedElementsDataStorage/Source/TypedElementsDataStorageUI/Private/TypedElementExportedTextWidget.cpp b/Engine/Plugins/Experimental/TypedElementsDataStorage/Source/TypedElementsDataStorageUI/Private/TypedElementExportedTextWidget.cpp
--- a/Engine/Plugins/Experimental/TypedElementsDataStorage/Source/TypedElementsDataStorageUI/Private/TypedElementExportedTextWidget.cpp
+++ b/Engine/Plugins/Experimental/TypedElementsDataStorage/Source/TypedElementsDataStorageUI/Private/TypedElementExportedTextWidget.cpp
@@ -110,12 +110,7 @@ TConstArrayView<const UScriptStruct*> FTypedElementExportedTextWidgetConstructor
 	return Columns;
 }
 
-bool FTypedElementExportedTextWidgetConstructor::CanBeReused() const
-{
-	return true;
-}
-
-TSharedPtr<SWidget> FTypedElementExportedTextWidgetConstructor::CreateWidget()
+TSharedPtr<SWidget> FTypedElementExportedTextWidgetConstructor::CreateWidget(TypedElementDataStorage::FMetaDataView Arguments)
 {
 	return SNew(STextBlock);
 }
diff --git a/Engine/Plugins/Experimental/TypedElementsDataStorage/Source/TypedElementsDataStorageUI/Private/TypedElementLabelWidget.cpp b/Engine/Plugins/Experimental/TypedElementsDataStorage/Source/TypedElementsDataStorageUI/Private/TypedElementLabelWidget.cpp
--- a/Engine/Plugins/Experimental/TypedElementsDataStorage/Source/TypedElementsDataStorageUI/Private/TypedElementLabelWidget.cpp
+++ b/Engine/Plugins/Experimental/TypedElementsDataStorage/Source/TypedElementsDataStorageUI/Private/TypedElementLabelWidget.cpp
@@ -150,14 +150,10 @@ TConstArrayView<const UScriptStruct*> FTypedElementLabelWidgetConstructor::GetAd
 	return Columns;
 }
 
-bool FTypedElementLabelWidgetConstructor::CanBeReused() const
+TSharedPtr<SWidget> FTypedElementLabelWidgetConstructor::CreateWidget(TypedElementDataStorage::FMetaDataView Arguments)
 {
-	return true;
-}
-
-TSharedPtr<SWidget> FTypedElementLabelWidgetConstructor::CreateWidget()
-{
-	return SNew(STextBlock);
+	const bool* IsEditable = Arguments.FindForColumn<FTypedElementLabelColumn>(TypedElementDataStorage::IsEditableName).TryGetExact<bool>();
+	return SNew(STextBlock).IsEnabled(IsEditable ? *IsEditable : false);
 }
 
 bool FTypedElementLabelWidgetConstructor::SetColumns(ITypedElementDataStorageInterface* DataStorage, TypedElementRowHandle Row)
diff --git a/Engine/Plugins/Experimental/TypedElementsDataStorage/Source/TypedElementsDataStorageUI/Private/TypedElementPackagePathWidget.cpp b/Engine/Plugins/Experimental/TypedElementsDataStorage/Source/TypedElementsDataStorageUI/Private/TypedElementPackagePathWidget.cpp
--- a/Engine/Plugins/Experimental/TypedElementsDataStorage/Source/TypedElementsDataStorageUI/Private/TypedElementPackagePathWidget.cpp
+++ b/Engine/Plugins/Experimental/TypedElementsDataStorage/Source/TypedElementsDataStorageUI/Private/TypedElementPackagePathWidget.cpp
@@ -39,12 +39,7 @@ FTypedElementPackagePathWidgetConstructor::FTypedElementPackagePathWidgetConstru
 {
 }
 
-bool FTypedElementPackagePathWidgetConstructor::CanBeReused() const
-{
-	return true;
-}
-
-TSharedPtr<SWidget> FTypedElementPackagePathWidgetConstructor::CreateWidget()
+TSharedPtr<SWidget> FTypedElementPackagePathWidgetConstructor::CreateWidget(TypedElementDataStorage::FMetaDataView Arguments)
 {
 	return SNew(STextBlock)
 		.OverflowPolicy(ETextOverflowPolicy::Ellipsis)
diff --git a/Engine/Plugins/Experimental/TypedElementsDataStorage/Source/TypedElementsDataStorageUI/Private/TypedElementTransformHeadsUpWidget.cpp b/Engine/Plugins/Experimental/TypedElementsDataStorage/Source/TypedElementsDataStorageUI/Private/TypedElementTransformHeadsUpWidget.cpp
--- a/Engine/Plugins/Experimental/TypedElementsDataStorage/Source/TypedElementsDataStorageUI/Private/TypedElementTransformHeadsUpWidget.cpp
+++ b/Engine/Plugins/Experimental/TypedElementsDataStorage/Source/TypedElementsDataStorageUI/Private/TypedElementTransformHeadsUpWidget.cpp
@@ -213,12 +213,7 @@ TConstArrayView<const UScriptStruct*> FTypedElementTransformHeadsUpWidgetConstru
 	return Columns;
 }
 
-bool FTypedElementTransformHeadsUpWidgetConstructor::CanBeReused() const
-{
-	return true;
-}
-
-TSharedPtr<SWidget> FTypedElementTransformHeadsUpWidgetConstructor::CreateWidget()
+TSharedPtr<SWidget> FTypedElementTransformHeadsUpWidgetConstructor::CreateWidget(TypedElementDataStorage::FMetaDataView Arguments)
 {
 	return SNew(STransformQuickDisplay);
 }
diff --git a/Engine/Plugins/Experimental/TypedElementsDataStorage/Source/TypedElementsDataStorageUI/Public/TypedElementCounterWidget.h b/Engine/Plugins/Experimental/TypedElementsDataStorage/Source/TypedElementsDataStorageUI/Public/TypedElementCounterWidget.h
--- a/Engine/Plugins/Experimental/TypedElementsDataStorage/Source/TypedElementsDataStorageUI/Public/TypedElementCounterWidget.h
+++ b/Engine/Plugins/Experimental/TypedElementsDataStorage/Source/TypedElementsDataStorageUI/Public/TypedElementCounterWidget.h
@@ -48,14 +48,13 @@ public:
 	~FTypedElementCounterWidgetConstructor() override = default;
 
 	TConstArrayView<const UScriptStruct*> GetAdditionalColumnsList() const override;
-	bool CanBeReused() const override;
 
 	FText ToolTipText{ NSLOCTEXT("TypedElementUI_CounterWidget", "Tooltip", "Shows the total number found in the editor.") };
 	FText LabelText{ NSLOCTEXT("TypedElementUI_CounterWidget", "Label", "Counted") };
 	TypedElementQueryHandle Query;
 
 protected:
-	TSharedPtr<SWidget> CreateWidget() override;
+	TSharedPtr<SWidget> CreateWidget(TypedElementDataStorage::FMetaDataView Arguments) override;
 	bool SetColumns(ITypedElementDataStorageInterface* DataStorage, TypedElementRowHandle Row) override;
 };
 
diff --git a/Engine/Plugins/Experimental/TypedElementsDataStorage/Source/TypedElementsDataStorageUI/Public/TypedElementExportedTextWidget.h b/Engine/Plugins/Experimental/TypedElementsDataStorage/Source/TypedElementsDataStorageUI/Public/TypedElementExportedTextWidget.h
--- a/Engine/Plugins/Experimental/TypedElementsDataStorage/Source/TypedElementsDataStorageUI/Public/TypedElementExportedTextWidget.h
+++ b/Engine/Plugins/Experimental/TypedElementsDataStorage/Source/TypedElementsDataStorageUI/Public/TypedElementExportedTextWidget.h
@@ -34,10 +34,9 @@ public:
 	~FTypedElementExportedTextWidgetConstructor() override = default;
 
 	TConstArrayView<const UScriptStruct*> GetAdditionalColumnsList() const override;
-	bool CanBeReused() const override;
 
 protected:
-	TSharedPtr<SWidget> CreateWidget() override;
+	TSharedPtr<SWidget> CreateWidget(TypedElementDataStorage::FMetaDataView Arguments) override;
 	bool FinalizeWidget(
 		ITypedElementDataStorageInterface* DataStorage,
 		ITypedElementDataStorageUiInterface* DataStorageUi,
diff --git a/Engine/Plugins/Experimental/TypedElementsDataStorage/Source/TypedElementsDataStorageUI/Public/TypedElementLabelWidget.h b/Engine/Plugins/Experimental/TypedElementsDataStorage/Source/TypedElementsDataStorageUI/Public/TypedElementLabelWidget.h
--- a/Engine/Plugins/Experimental/TypedElementsDataStorage/Source/TypedElementsDataStorageUI/Public/TypedElementLabelWidget.h
+++ b/Engine/Plugins/Experimental/TypedElementsDataStorage/Source/TypedElementsDataStorageUI/Public/TypedElementLabelWidget.h
@@ -33,11 +33,10 @@ public:
 	~FTypedElementLabelWidgetConstructor() override = default;
 
 	TConstArrayView<const UScriptStruct*> GetAdditionalColumnsList() const override;
-	bool CanBeReused() const override;
 
 protected:
 	explicit FTypedElementLabelWidgetConstructor(const UScriptStruct* InTypeInfo);
-	TSharedPtr<SWidget> CreateWidget() override;
+	TSharedPtr<SWidget> CreateWidget(TypedElementDataStorage::FMetaDataView Arguments) override;
 	bool SetColumns(ITypedElementDataStorageInterface* DataStorage, TypedElementRowHandle Row) override;
 	bool FinalizeWidget(ITypedElementDataStorageInterface* DataStorage, ITypedElementDataStorageUiInterface* DataStorageUi,
 		TypedElementRowHandle Row, const TSharedPtr<SWidget>& Widget) override;
diff --git a/Engine/Plugins/Experimental/TypedElementsDataStorage/Source/TypedElementsDataStorageUI/Public/TypedElementPackagePathWidget.h b/Engine/Plugins/Experimental/TypedElementsDataStorage/Source/TypedElementsDataStorageUI/Public/TypedElementPackagePathWidget.h
--- a/Engine/Plugins/Experimental/TypedElementsDataStorage/Source/TypedElementsDataStorageUI/Public/TypedElementPackagePathWidget.h
+++ b/Engine/Plugins/Experimental/TypedElementsDataStorage/Source/TypedElementsDataStorageUI/Public/TypedElementPackagePathWidget.h
@@ -31,12 +31,10 @@ public:
 	FTypedElementPackagePathWidgetConstructor();
 	~FTypedElementPackagePathWidgetConstructor() override = default;
 
-	bool CanBeReused() const override;
-
 protected:
 	explicit FTypedElementPackagePathWidgetConstructor(const UScriptStruct* InTypeInfo);
 
-	TSharedPtr<SWidget> CreateWidget() override;
+	TSharedPtr<SWidget> CreateWidget(TypedElementDataStorage::FMetaDataView Arguments) override;
 	bool FinalizeWidget(ITypedElementDataStorageInterface* DataStorage, ITypedElementDataStorageUiInterface* DataStorageUi,
 		TypedElementRowHandle Row, const TSharedPtr<SWidget>& Widget) override;
 };
diff --git a/Engine/Plugins/Experimental/TypedElementsDataStorage/Source/TypedElementsDataStorageUI/Public/TypedElementTransformHeadsUpWidget.h b/Engine/Plugins/Experimental/TypedElementsDataStorage/Source/TypedElementsDataStorageUI/Public/TypedElementTransformHeadsUpWidget.h
--- a/Engine/Plugins/Experimental/TypedElementsDataStorage/Source/TypedElementsDataStorageUI/Public/TypedElementTransformHeadsUpWidget.h
+++ b/Engine/Plugins/Experimental/TypedElementsDataStorage/Source/TypedElementsDataStorageUI/Public/TypedElementTransformHeadsUpWidget.h
@@ -39,10 +39,9 @@ public:
 	~FTypedElementTransformHeadsUpWidgetConstructor() override = default;
 
 	TConstArrayView<const UScriptStruct*> GetAdditionalColumnsList() const override;
-	bool CanBeReused() const override;
 
 protected:
-	TSharedPtr<SWidget> CreateWidget() override;
+	TSharedPtr<SWidget> CreateWidget(TypedElementDataStorage::FMetaDataView Arguments) override;
 	bool FinalizeWidget(
 		ITypedElementDataStorageInterface* DataStorage,
 		ITypedElementDataStorageUiInterface* DataStorageUi,
diff --git a/Engine/Source/Runtime/TypedElementFramework/Private/Elements/Framework/TypedElementMetaData.cpp b/Engine/Source/Runtime/TypedElementFramework/Private/Elements/Framework/TypedElementMetaData.cpp
new file mode 100644
--- /dev/null
+++ b/Engine/Source/Runtime/TypedElementFramework/Private/Elements/Framework/TypedElementMetaData.cpp
@@ -0,0 +1,160 @@
+// Copyright Epic Games, Inc. All Rights Reserved.
+
+#include "Elements/Framework/TypedElementMetaData.h"
+
+#include "Elements/Common/TypedElementQueryDescription.h"
+#include "UObject/Class.h"
+#include "UObject/UnrealType.h"
+
+namespace TypedElementDataStorage
+{
+	template<class... Ts>
+	struct TOverloaded : Ts... 
+	{ 
+		using Ts::operator()...;
+	};
+	
+	template<class... Ts>
+	TOverloaded(Ts...) -> TOverloaded<Ts...>;
+
+	//
+	// FMetaDataBase
+	//
+
+	FMetaDataEntryView FMetaDataBase::Find(FName Name) const
+	{
+		if (const MetaDataType* Result = ImmutableData.Find(Name))
+		{
+			return FMetaDataEntryView(*Result);
+		}
+
+		if (const MetaDataType* Result = MutableData.Find(Name))
+		{
+			return FMetaDataEntryView(*Result);
+		}
+
+		return FMetaDataEntryView();
+	}
+
+	void FMetaDataBase::Shrink()
+	{
+		ImmutableData.Shrink();
+		MutableData.Shrink();
+	}
+
+
+
+	//
+	// FColumnMetadata
+	//
+
+	FColumnMetaData::FColumnMetaData(const UScriptStruct* InColumnType, EFlags InFlags)
+		: ColumnType(InColumnType)
+		, Flags(InFlags)
+	{}
+
+	FMetaDataEntryView FColumnMetaData::Find(FName Name) const
+	{
+		if (Name == IsEditableName)
+		{
+			return FMetaDataEntryView((Flags & EFlags::IsMutable) != EFlags::None);
+		}
+		else if (Name == IsConstName)
+		{
+			return FMetaDataEntryView(!(Flags & EFlags::IsMutable));
+		}
+
+		FMetaDataEntryView Result = FMetaDataBase::Find(Name);
+
+#if WITH_EDITORONLY_DATA
+		if (ColumnType && !Result.IsSet())
+		{
+			if (const FString* FoundMetaData = ColumnType->FindMetaData(Name))
+			{
+				return FMetaDataEntryView(*FoundMetaData);
+			}
+		}
+#endif
+
+		return Result;
+	}
+
+
+	//
+	// FMetaDataEntryView
+	// 
+
+	FMetaDataEntryView::FMetaDataEntryView()
+		: DataView(TInPlaceType<FEmptyVariantState>())
+	{}
+
+	FMetaDataEntryView::FMetaDataEntryView(const MetaDataType& MetaData)
+	{
+		Visit(TOverloaded
+			{
+				[this](const auto& Value)
+				{
+					using TargetType = typename TDecay<decltype(Value)>::Type;
+					DataView.Emplace<TargetType>(Value);
+				},
+				[this](const FString& String)
+				{
+					DataView.Emplace<const FString*>(&String);
+				}
+			}, MetaData);
+	}
+
+	FMetaDataEntryView::FMetaDataEntryView(const FString& MetaDataString)
+		: DataView(TInPlaceType<const FString*>(), &MetaDataString)
+	{}
+
+	bool FMetaDataEntryView::IsSet() const
+	{
+		return !DataView.IsType<FEmptyVariantState>();
+	}
+
+
+
+	//
+	// FMetaDataView
+	//
+
+	FMetaDataView::FMetaDataView(const TypedElementDataStorage::FQueryDescription& InQuery)
+		: Query(&InQuery)
+	{}
+
+	FMetaDataView::FMetaDataView(const FMetaData& InQueryWideMetaData)
+		: QueryWideMetaData(&InQueryWideMetaData)
+	{}
+
+	FMetaDataView::FMetaDataView(const TypedElementDataStorage::FQueryDescription& InQuery, const FMetaData& InQueryWideMetaData)
+		: Query(&InQuery)
+		, QueryWideMetaData(&InQueryWideMetaData)
+	{}
+
+	FMetaDataEntryView FMetaDataView::FindGeneric(FName AttributeName) const
+	{
+		return QueryWideMetaData ? QueryWideMetaData->Find(AttributeName) : FMetaDataEntryView();
+	}
+
+	FMetaDataEntryView FMetaDataView::FindForColumn(TWeakObjectPtr<const UScriptStruct> Column, FName AttributeName, ESearchScope Scope) const
+	{
+		if (Query)
+		{
+			// There are typically a small number of columns in a query, so a linear search if often fast enough and can even
+			// be faster than a map.
+			int32 Index = 0;
+			if (Query->SelectionTypes.Find(Column, Index))
+			{
+				FMetaDataEntryView Result = Query->SelectionMetaData[Index].Find(AttributeName);
+				if (Result.IsSet())
+				{
+					return Result;
+				}
+			}
+		}
+		return Scope == ESearchScope::FallbackOnGeneric && QueryWideMetaData
+			? QueryWideMetaData->Find(AttributeName)
+			: FMetaDataEntryView();
+	}
+} // namespace TypedElementDataStorage
diff --git a/Engine/Source/Runtime/TypedElementFramework/Private/Elements/Framework/TypedElementQueryBuilder.cpp b/Engine/Source/Runtime/TypedElementFramework/Private/Elements/Framework/TypedElementQueryBuilder.cpp
--- a/Engine/Source/Runtime/TypedElementFramework/Private/Elements/Framework/TypedElementQueryBuilder.cpp
+++ b/Engine/Source/Runtime/TypedElementFramework/Private/Elements/Framework/TypedElementQueryBuilder.cpp
@@ -2,6 +2,8 @@
 
 #include "Elements/Framework/TypedElementQueryBuilder.h"
 
+#include "Elements/Framework/TypedElementMetaData.h"
+
 namespace TypedElementQueryBuilder
 {
 	const UScriptStruct* Type(FTopLevelAssetPath Name)
@@ -190,6 +192,11 @@ namespace TypedElementQueryBuilder
 		Query->Callback.AfterGroups.Shrink();
 		Query->SelectionTypes.Shrink();
 		Query->SelectionAccessTypes.Shrink();
+		for (TypedElementDataStorage::FColumnMetaData& Metadata : Query->SelectionMetaData)
+		{
+			Metadata.Shrink();
+		}
+		Query->SelectionMetaData.Shrink();
 		Query->ConditionTypes.Shrink();
 		Query->ConditionOperators.Shrink();
 		Query->DependencyTypes.Shrink();
@@ -309,11 +316,18 @@ namespace TypedElementQueryBuilder
 		checkf(Target, TEXT("The Select section in the Typed Elements query builder doesn't support nullptrs as Read-Only input."));
 		Query.SelectionTypes.Emplace(Target);
 		Query.SelectionAccessTypes.Emplace(ITypedElementDataStorageInterface::EQueryAccessType::ReadOnly);
+		Query.SelectionMetaData.Emplace(Target, TypedElementDataStorage::FColumnMetaData::EFlags::None);
+		
 		return *this;
 	}
 
-	Select& Select::ReadOnly(std::initializer_list<const UScriptStruct*> Targets)
+	Select& Select::ReadOnly(TConstArrayView<const UScriptStruct*> Targets)
 	{
+		int32 NewCount = Query.SelectionTypes.Num() + Targets.Num();
+		Query.SelectionTypes.Reserve(NewCount);
+		Query.SelectionAccessTypes.Reserve(NewCount);
+		Query.SelectionMetaData.Reserve(NewCount);
+
 		for (const UScriptStruct* Target : Targets)
 		{
 			ReadOnly(Target);
@@ -326,11 +340,17 @@ namespace TypedElementQueryBuilder
 		checkf(Target, TEXT("The Select section in the Typed Elements query builder doesn't support nullptrs as Read/Write input."));
 		Query.SelectionTypes.Emplace(Target);
 		Query.SelectionAccessTypes.Emplace(ITypedElementDataStorageInterface::EQueryAccessType::ReadWrite);
+		Query.SelectionMetaData.Emplace(Target, TypedElementDataStorage::FColumnMetaData::EFlags::IsMutable);
 		return *this;
 	}
 
-	Select& Select::ReadWrite(std::initializer_list<const UScriptStruct*> Targets)
+	Select& Select::ReadWrite(TConstArrayView<const UScriptStruct*> Targets)
 	{
+		int32 NewCount = Query.SelectionTypes.Num() + Targets.Num();
+		Query.SelectionTypes.Reserve(NewCount);
+		Query.SelectionAccessTypes.Reserve(NewCount);
+		Query.SelectionMetaData.Reserve(NewCount);
+
 		for (const UScriptStruct* Target : Targets)
 		{
 			ReadWrite(Target);
diff --git a/Engine/Source/Runtime/TypedElementFramework/Private/Elements/Interfaces/TypedElementDataStorageUiInterface.cpp b/Engine/Source/Runtime/TypedElementFramework/Private/Elements/Interfaces/TypedElementDataStorageUiInterface.cpp
--- a/Engine/Source/Runtime/TypedElementFramework/Private/Elements/Interfaces/TypedElementDataStorageUiInterface.cpp
+++ b/Engine/Source/Runtime/TypedElementFramework/Private/Elements/Interfaces/TypedElementDataStorageUiInterface.cpp
@@ -21,39 +21,30 @@ TConstArrayView<const UScriptStruct*> FTypedElementWidgetConstructor::GetAdditio
 
 bool FTypedElementWidgetConstructor::CanBeReused() const
 {
-	return false;
+	return true;
 }
 
 TSharedPtr<SWidget> FTypedElementWidgetConstructor::Construct(
 	TypedElementRowHandle Row,
 	ITypedElementDataStorageInterface* DataStorage,
 	ITypedElementDataStorageUiInterface* DataStorageUi,
-	TConstArrayView<TypedElement::ColumnUtils::Argument> Arguments)
+	TypedElementDataStorage::FMetaDataView Arguments)
 {
-	if (ApplyArguments(Arguments))
+	TSharedPtr<SWidget> Widget = CreateWidget(Arguments);
+	if (Widget)
 	{
-		TSharedPtr<SWidget> Widget = CreateWidget();
-		if (Widget)
+		DataStorage->GetColumn<FTypedElementSlateWidgetReferenceColumn>(Row)->Widget = Widget;
+		if (SetColumns(DataStorage, Row))
 		{
-			DataStorage->GetColumn<FTypedElementSlateWidgetReferenceColumn>(Row)->Widget = Widget;
-			if (SetColumns(DataStorage, Row))
+			if (FinalizeWidget(DataStorage, DataStorageUi, Row, Widget))
 			{
-				if (FinalizeWidget(DataStorage, DataStorageUi, Row, Widget))
-				{
-					return Widget;
-				}
+				return Widget;
 			}
 		}
 	}
 	return nullptr;
 }
 
-bool FTypedElementWidgetConstructor::ApplyArguments(TConstArrayView<TypedElement::ColumnUtils::Argument> Arguments)
-{
-	SetColumnValues(*this, Arguments);
-	return true;
-}
-
 bool FTypedElementWidgetConstructor::SetColumns(ITypedElementDataStorageInterface* DataStorage, TypedElementRowHandle Row)
 {
 	return true;
diff --git a/Engine/Source/Runtime/TypedElementFramework/Public/Elements/Common/TypedElementQueryDescription.h b/Engine/Source/Runtime/TypedElementFramework/Public/Elements/Common/TypedElementQueryDescription.h
--- a/Engine/Source/Runtime/TypedElementFramework/Public/Elements/Common/TypedElementQueryDescription.h
+++ b/Engine/Source/Runtime/TypedElementFramework/Public/Elements/Common/TypedElementQueryDescription.h
@@ -7,6 +7,7 @@
 #include "Elements/Common/TypedElementHandles.h"
 #include "Elements/Common/TypedElementQueryTypes.h"
 #include "Elements/Interfaces/TypedElementQueryStorageInterfaces.h"
+#include "Elements/Framework/TypedElementMetaData.h"
 #include "Templates/Function.h"
 #include "UObject/NameTypes.h"
 #include "UObject/WeakObjectPtr.h"
@@ -87,6 +88,7 @@ namespace TypedElementDataStorage
 
 		TArray<TWeakObjectPtr<const UScriptStruct>, TInlineAllocator<NumInlineSelections>> SelectionTypes;
 		TArray<EQueryAccessType, TInlineAllocator<NumInlineSelections>> SelectionAccessTypes;
+		TArray<FColumnMetaData, TInlineAllocator<NumInlineSelections>> SelectionMetaData;
 
 		TArray<EOperatorType, TInlineAllocator<NumInlineConditions>> ConditionTypes;
 		TArray<FOperator, TInlineAllocator<NumInlineConditions>> ConditionOperators;
diff --git a/Engine/Source/Runtime/TypedElementFramework/Public/Elements/Framework/TypedElementMetaData.h b/Engine/Source/Runtime/TypedElementFramework/Public/Elements/Framework/TypedElementMetaData.h
new file mode 100644
--- /dev/null
+++ b/Engine/Source/Runtime/TypedElementFramework/Public/Elements/Framework/TypedElementMetaData.h
@@ -0,0 +1,132 @@
+// Copyright Epic Games, Inc. All Rights Reserved.
+
+#pragma once
+
+#include <type_traits>
+#include "CoreMinimal.h"
+#include "Containers/ArrayView.h"
+#include "Containers/Map.h"
+#include "Misc/EnumClassFlags.h"
+#include "Misc/TVariant.h"
+
+class UScriptStruct;
+
+namespace TypedElementDataStorage
+{
+	struct FQueryDescription;
+
+	inline static const FName IsEditableName(TEXT("IsEditable"));
+	inline static const FName IsConstName(TEXT("IsConst"));
+
+	using MetaDataType = TVariant<bool, uint64, int64, double, FString>;
+	using MetaDataTypeView = TVariant<FEmptyVariantState, bool, uint64, int64, double, const FString*>;
+
+	/**
+	 * Short lived view of single entry in the meta data container.
+	 */
+	class FMetaDataEntryView final
+	{
+	public:
+		TYPEDELEMENTFRAMEWORK_API FMetaDataEntryView();
+		TYPEDELEMENTFRAMEWORK_API explicit FMetaDataEntryView(const MetaDataType& MetaData);
+		TYPEDELEMENTFRAMEWORK_API explicit FMetaDataEntryView(const FString& MetaDataString);
+		/** Explicit constructor. The provided type must match exactly to one of the available stored types. */
+		template<typename T>
+		explicit FMetaDataEntryView(T&& MetaDataValue);
+
+		/** Returns true if set to a value, otherwise false. */
+		TYPEDELEMENTFRAMEWORK_API bool IsSet() const;
+		/** Checks if the stored value matches the requested type. */
+		template<typename T>
+		bool IsType() const;
+		/** Returns the value if the requested type matches exactly with the stored type, otherwise returns a nullptr. */
+		template<typename T>
+		const T* TryGetExact() const;
+
+	private:
+		MetaDataTypeView DataView;
+	};
+
+	/**
+	 * Base class to store meta data for use within the Typed Elements Data Storage.
+	 */
+	class FMetaDataBase
+	{
+	public:
+		template<typename T>
+		bool AddImmutableData(FName Name, T&& Value);
+		template<typename T>
+		void AddOrSetMutableData(FName Name, T&& Value);
+
+		TYPEDELEMENTFRAMEWORK_API virtual FMetaDataEntryView Find(FName Name) const;
+
+		TYPEDELEMENTFRAMEWORK_API virtual void Shrink();
+
+	protected:
+		FMetaDataBase() = default;
+		virtual ~FMetaDataBase() = default;
+
+		/** Data that can be added once but can't be changed afterwards. Values here always take priority over other values. */
+		TMap<FName, MetaDataType> ImmutableData;
+		/** Data that can be added and can have their value updated afterwards. */
+		TMap<FName, MetaDataType> MutableData;
+	};
+
+	/** General storage for meta data for the Typed Elements Data Storage. */
+	class FMetaData final : public FMetaDataBase {};
+
+	/**
+	 * Meta data that's specifically associated with a single column.
+	 */
+	class FColumnMetaData final : public FMetaDataBase
+	{
+	public:
+		enum class EFlags
+		{
+			None = 0,
+			IsMutable = 1 << 0
+		};
+
+		FColumnMetaData() = default;
+		TYPEDELEMENTFRAMEWORK_API FColumnMetaData(const UScriptStruct* InColumnType, EFlags InFlags);
+
+		TYPEDELEMENTFRAMEWORK_API FMetaDataEntryView Find(FName Name) const override;
+
+	private:
+		/** If set, properties on the column will also be included if a value isn't found in the (im)mutable data map. */
+		const UScriptStruct* ColumnType{ nullptr };
+		/** Flags indicating the behavior of the column */
+		EFlags Flags = EFlags::None;
+	};
+
+	/**
+	 * Short lived view of a meta data container optionally associated with a query.
+	 */
+	class FMetaDataView final
+	{
+	public:
+		enum class ESearchScope
+		{
+			FallbackOnGeneric, // If the attribute isn't found on the column or the column isn't found, search the generic meta data for the attribute.
+			ColumnOnly // Only use the data on the column.
+		};
+
+		FMetaDataView() = default;
+		TYPEDELEMENTFRAMEWORK_API FMetaDataView(const TypedElementDataStorage::FQueryDescription& InQuery); // Deliberately avoided "explicit".
+		TYPEDELEMENTFRAMEWORK_API FMetaDataView(const FMetaData& InQueryWideMetaData); // Deliberately avoided "explicit".
+		TYPEDELEMENTFRAMEWORK_API FMetaDataView(
+			const TypedElementDataStorage::FQueryDescription& InQuery, const FMetaData& InQueryWideMetaData);
+
+		TYPEDELEMENTFRAMEWORK_API FMetaDataEntryView FindGeneric(FName AttributeName) const;
+		TYPEDELEMENTFRAMEWORK_API FMetaDataEntryView FindForColumn(
+			TWeakObjectPtr<const UScriptStruct> Column, FName AttributeName, ESearchScope Scope = ESearchScope::ColumnOnly) const;
+		template<typename Column>
+		FMetaDataEntryView FindForColumn(FName AttributeName, ESearchScope Scope = ESearchScope::ColumnOnly) const;
+
+	private:
+		const TypedElementDataStorage::FQueryDescription* Query{ nullptr };
+		const FMetaData* QueryWideMetaData{ nullptr };
+	};
+} // TypedElementDataStorage
+
+#include "Elements/Framework/TypedElementMetaData.inl"
\ No newline at end of file
diff --git a/Engine/Source/Runtime/TypedElementFramework/Public/Elements/Framework/TypedElementMetaData.inl b/Engine/Source/Runtime/TypedElementFramework/Public/Elements/Framework/TypedElementMetaData.inl
new file mode 100644
--- /dev/null
+++ b/Engine/Source/Runtime/TypedElementFramework/Public/Elements/Framework/TypedElementMetaData.inl
@@ -0,0 +1,106 @@
+// Copyright Epic Games, Inc. All Rights Reserved.
+
+namespace TypedElementDataStorage
+{
+	namespace Internal
+	{
+		template<typename TContainer, typename T>
+		void AddDataToContainer(TContainer& Container, FName Name, T Value)
+		{
+			using BaseType = std::remove_const_t<std::remove_reference_t<T>>;
+
+			if constexpr (std::is_arithmetic_v<BaseType>)
+			{
+				if constexpr (std::is_integral_v<BaseType>)
+				{
+					if constexpr (std::is_signed_v<BaseType>)
+					{
+						Container.Add(Name, MetaDataType(TInPlaceType<int64>(), static_cast<int64>(Value)));
+					}
+					else
+					{
+						Container.Add(Name, MetaDataType(TInPlaceType<uint64>(), static_cast<uint64>(Value)));
+					}
+				}
+				else if constexpr (std::is_floating_point_v<BaseType>)
+				{
+					Container.Add(Name, MetaDataType(TInPlaceType<double>(), static_cast<double>(Value)));
+				}
+				else
+				{
+					Container.Add(Name, MetaDataType(TInPlaceType<bool>(), static_cast<bool>(Value)));
+				}
+			}
+			else if constexpr (std::is_convertible_v<BaseType, FString>)
+			{
+				Container.Add(Name, MetaDataType(TInPlaceType<FString>(), Forward<T>(Value)));
+			}
+			else
+			{
+				// !std::is_same_v<T, T> is used to force evaluation of the type so it's not always called, which would happen if "false" is used instead.
+				static_assert(!std::is_same_v<T, T>, "Provided meta data type can't be stored in the Typed Element Data Storage.");
+			}
+		}
+	}
+
+	ENUM_CLASS_FLAGS(TypedElementDataStorage::FColumnMetaData::EFlags)
+
+		template<typename T>
+	bool FMetaDataBase::AddImmutableData(FName Name, T&& Value)
+	{
+		if (!ImmutableData.Contains(Name))
+		{
+			Internal::AddDataToContainer(ImmutableData, Name, Forward<T>(Value));
+			return true;
+		}
+		else
+		{
+			return false;
+		}
+	}
+
+	template<typename T>
+	void FMetaDataBase::AddOrSetMutableData(FName Name, T&& Value)
+	{
+		Internal::AddDataToContainer(MutableData, Name, Forward<T>(Value));
+	}
+
+
+
+	//
+	// FMetaDataEntryView
+	//
+
+	template<typename T>
+	FMetaDataEntryView::FMetaDataEntryView(T&& MetaDataValue)
+		: DataView(TInPlaceType<T>(), Forward<T>(MetaDataValue))
+	{}
+
+	template<typename T>
+	bool FMetaDataEntryView::IsType() const
+	{
+		return DataView.IsType<T>();
+	}
+
+	template<typename T>
+	const T* FMetaDataEntryView::TryGetExact() const
+	{
+		if constexpr (std::is_same_v<T, const FString&>)
+		{
+			return DataView.TryGet<const FString*>();
+		}
+		return DataView.TryGet<T>();
+	}
+
+
+
+	//
+	// FMetaDataView
+	//
+
+	template<typename Column>
+	FMetaDataEntryView FMetaDataView::FindForColumn(FName AttributeName, ESearchScope Scope) const
+	{
+		return FindForColumn(Column::StaticStruct(), AttributeName, Scope);
+	}
+} // namespace TypedElementDataStorage
diff --git a/Engine/Source/Runtime/TypedElementFramework/Public/Elements/Framework/TypedElementQueryBuilder.h b/Engine/Source/Runtime/TypedElementFramework/Public/Elements/Framework/TypedElementQueryBuilder.h
--- a/Engine/Source/Runtime/TypedElementFramework/Public/Elements/Framework/TypedElementQueryBuilder.h
+++ b/Engine/Source/Runtime/TypedElementFramework/Public/Elements/Framework/TypedElementQueryBuilder.h
@@ -295,11 +295,11 @@ namespace TypedElementQueryBuilder
 		template<typename... TargetTypes>
 		Select& ReadOnly();
 		TYPEDELEMENTFRAMEWORK_API Select& ReadOnly(const UScriptStruct* Target);
-		TYPEDELEMENTFRAMEWORK_API Select& ReadOnly(std::initializer_list<const UScriptStruct*> Targets);
+		TYPEDELEMENTFRAMEWORK_API Select& ReadOnly(TConstArrayView<const UScriptStruct*> Targets);
 		template<typename... TargetTypes>
 		Select& ReadWrite();
 		TYPEDELEMENTFRAMEWORK_API Select& ReadWrite(const UScriptStruct* Target);
-		TYPEDELEMENTFRAMEWORK_API Select& ReadWrite(std::initializer_list<const UScriptStruct*> Targets);
+		TYPEDELEMENTFRAMEWORK_API Select& ReadWrite(TConstArrayView<const UScriptStruct*> Targets);
 
 		TYPEDELEMENTFRAMEWORK_API ITypedElementDataStorageInterface::FQueryDescription&& Compile();
 		TYPEDELEMENTFRAMEWORK_API FSimpleQuery Where();
diff --git a/Engine/Source/Runtime/TypedElementFramework/Public/Elements/Interfaces/TypedElementDataStorageUiInterface.h b/Engine/Source/Runtime/TypedElementFramework/Public/Elements/Interfaces/TypedElementDataStorageUiInterface.h
--- a/Engine/Source/Runtime/TypedElementFramework/Public/Elements/Interfaces/TypedElementDataStorageUiInterface.h
+++ b/Engine/Source/Runtime/TypedElementFramework/Public/Elements/Interfaces/TypedElementDataStorageUiInterface.h
@@ -3,8 +3,8 @@
 #pragma once
 
 #include "Containers/ContainersFwd.h"
-#include "Elements/Framework/TypedElementColumnUtils.h"
 #include "Elements/Interfaces/TypedElementDataStorageInterface.h"
+#include "Elements/Framework/TypedElementMetaData.h"
 #include "Templates/UnrealTypeTraits.h"
 #include "UObject/Interface.h"
 
@@ -15,13 +15,7 @@ class SWidget;
 
 /**
  * Base class used to construct Typed Element widgets with.
- * It's recommended to expose any construction variables as properties so they
- * can either be set by a user or set using the passed in arguments. The
- * Arguments can be directly used in case complex operations need to be done that
- * prevent automatically setting construction variables.
- * See below for the options to register a constructor with the Data Storage. For
- * either registration case a new instance/copy of the constructor is created so
- * arguments can be safely applied.
+ * See below for the options to register a constructor with the Data Storage.
  */
 USTRUCT()
 struct FTypedElementWidgetConstructor
@@ -50,7 +44,7 @@ public:
 
 	/**
 	 * Constructs the widget according to the provided information. Information is collected by calling
-	 * the below functions ApplyArguments, CreateWidget and AddColumns. It's recommended to overload those
+	 * the below functions CreateWidget and AddColumns. It's recommended to overload those
 	 * functions to build widgets according to a standard recipe and to reduce the amount of code needed.
 	 * If a complexer situation is called for this function can also be directly overwritten.
 	 */
@@ -58,17 +52,12 @@ public:
 		TypedElementRowHandle Row, /** The row the widget will be stored in. */
 		ITypedElementDataStorageInterface* DataStorage,
 		ITypedElementDataStorageUiInterface* DataStorageUi,
-		TConstArrayView<TypedElement::ColumnUtils::Argument> Arguments);
+		TypedElementDataStorage::FMetaDataView Arguments);
 
 protected:
-	/**
-	 * Uses the type system to apply the provided arguments to constructor's properties. Overwrite this 
-	 * function if there are non-properties that can be configured through the arguments or if the configuration
-	 * contains arguments that are too complex to initialize (fully) through the type system.
-	 */
-	TYPEDELEMENTFRAMEWORK_API virtual bool ApplyArguments(TConstArrayView<TypedElement::ColumnUtils::Argument> Arguments);
 	/** Create a new instance of the target widget. This is a required function. */
-	TYPEDELEMENTFRAMEWORK_API virtual TSharedPtr<SWidget> CreateWidget() PURE_VIRTUAL(FTypedElementWidgetConstructor::CreateWidget, return nullptr; );
+	TYPEDELEMENTFRAMEWORK_API virtual TSharedPtr<SWidget> CreateWidget(
+		TypedElementDataStorage::FMetaDataView Arguments) PURE_VIRTUAL(FTypedElementWidgetConstructor::CreateWidget, return nullptr; );
 	/** Set any values in columns if needed. The columns provided through GetAdditionalColumnsList() will have already been created. */
 	TYPEDELEMENTFRAMEWORK_API virtual bool SetColumns(ITypedElementDataStorageInterface* DataStorage, TypedElementRowHandle Row);
 	/** 
@@ -183,25 +172,25 @@ public:
 	
 	/** Creates widget constructors for the requested purpose. */
 	virtual void CreateWidgetConstructors(FName Purpose, 
-		TConstArrayView<TypedElement::ColumnUtils::Argument> Arguments, const WidgetConstructorCallback& Callback) = 0;
+		TypedElementDataStorage::FMetaDataView Arguments, const WidgetConstructorCallback& Callback) = 0;
 	/** 
 	 * Finds matching widget constructors for provided columns, preferring longer matches over shorter matches.
 	 * The provided list of columns will be updated to contain all columns that couldn't be matched.
 	 */
 	virtual void CreateWidgetConstructors(FName Purpose, EMatchApproach MatchApproach, TArray<TWeakObjectPtr<const UScriptStruct>>& Columns,
-		TConstArrayView<TypedElement::ColumnUtils::Argument> Arguments, const WidgetConstructorCallback& Callback) = 0;
+		TypedElementDataStorage::FMetaDataView Arguments, const WidgetConstructorCallback& Callback) = 0;
 
 	/**
 	 * Creates all the widgets registered under the provided name. This may be a large number of widgets for a wide name
 	 * or exactly one when the exact name of the widget is registered. Arguments can be provided, but widgets are free
 	 * to ignore them.
 	 */
-	virtual void ConstructWidgets(FName Purpose, TConstArrayView<TypedElement::ColumnUtils::Argument> Arguments, 
+	virtual void ConstructWidgets(FName Purpose, TypedElementDataStorage::FMetaDataView Arguments,
 		const WidgetCreatedCallback& ConstructionCallback) = 0;
 
 	/** Creates a single widget using the provided constructor. Arguments can optionally be used to intialize the constructor. */
 	virtual TSharedPtr<SWidget> ConstructWidget(TypedElementRowHandle Row, FTypedElementWidgetConstructor& Constructor,
-		TConstArrayView<TypedElement::ColumnUtils::Argument> Arguments) = 0;
+		TypedElementDataStorage::FMetaDataView Arguments) = 0;
 
 	/** Calls the provided callback for all known registered widget purposes. */
 	virtual void ListWidgetPurposes(const WidgetPurposeCallback& Callback) const = 0;
