Cleanup of the memento and object reinstatiation system of TEDS
* Split object reinstantiation and column mementoization into separate systems.   Object reinstantion is part of the compatibility layer while the memento system can be reused for other purposes.
* Moved memento row creation to just before the target rows are about to be destroyed.  This saves memory by avoiding creating mementos which might never be used.
* Removed the temporary OnAdded/Removed delegates from the Compatibilty layer - this improves on that by hooking into the AddObject path using an Observer.
* MementoSystem is no longer a factory, instead it is a first-class citizen of the TypedElementDataStorageModule.  Similar for the ObjectReinstanceManager (new).
* Enables memento and object reinstantiation by default

diff --git a/Engine/Plugins/Experimental/TypedElementsDataStorage/Source/TypedElementsDataStorage/Private/Compatibility/TypedElementObjectReinstancingManager.cpp b/Engine/Plugins/Experimental/TypedElementsDataStorage/Source/TypedElementsDataStorage/Private/Compatibility/TypedElementObjectReinstancingManager.cpp
new file mode 100644
--- /dev/null
+++ b/Engine/Plugins/Experimental/TypedElementsDataStorage/Source/TypedElementsDataStorage/Private/Compatibility/TypedElementObjectReinstancingManager.cpp
@@ -0,0 +1,99 @@
+﻿// Copyright Epic Games, Inc. All Rights Reserved.
+
+#include "TypedElementObjectReinstancingManager.h"
+
+#include "Elements/Framework/TypedElementQueryBuilder.h"
+#include "HAL/IConsoleManager.h"
+#include "Memento/TypedElementMementoRowTypes.h"
+#include "TypedElementDatabase.h"
+#include "TypedElementDatabaseCompatibility.h"
+#include "Elements/Columns/TypedElementCompatibilityColumns.h"
+
+UTypedElementObjectReinstancingManager::UTypedElementObjectReinstancingManager()
+	: UObjectAddedObserverHandle(TypedElementInvalidQueryHandle)
+	, MementoRowBaseTable(TypedElementInvalidTableHandle)
+{
+}
+
+void UTypedElementObjectReinstancingManager::Initialize(UTypedElementDatabase& InDatabase, UTypedElementDatabaseCompatibility& InDataStorageCompatibility, UTypedElementMementoSystem& InMementoSystem)
+{
+	Database = &InDatabase;
+	DataStorageCompatibility = &InDataStorageCompatibility;
+	MementoSystem = &InMementoSystem;
+
+	ReinstancingCallbackHandle = FCoreUObjectDelegates::OnObjectsReinstanced.AddUObject(this, &UTypedElementObjectReinstancingManager::HandleOnObjectsReinstanced);
+	
+	RegisterQueries();
+}
+
+void UTypedElementObjectReinstancingManager::Deinitialize()
+{
+	UnregisterQueries();
+	
+	FCoreUObjectDelegates::OnObjectsReinstanced.Remove(ReinstancingCallbackHandle);
+
+	MementoSystem = nullptr;
+	DataStorageCompatibility = nullptr;
+	Database = nullptr;
+}
+
+void UTypedElementObjectReinstancingManager::RegisterQueries()
+{
+	using namespace TypedElementQueryBuilder;
+	
+	UObjectAddedObserverHandle = Database->RegisterQuery(
+		Select(
+			TEXT("Object Reinstancing: Added UObjectColumn"),
+			FObserver::OnAdd<FTypedElementUObjectColumn>(),
+			[this](TypedElementDataStorage::IQueryContext& Context, TypedElementRowHandle Row, const FTypedElementUObjectColumn& UObjectColumn)
+			{
+				TypedElementRowHandle Memento;
+				if (NewInstanceToMementoMap.RemoveAndCopyValue(UObjectColumn.Object.Get(), Memento))
+				{
+					Context.AddColumn<FTypedElementMementoReinstanceTarget>(Memento, FTypedElementMementoReinstanceTarget{.Target = Row});
+				}
+			})
+			.Compile());
+
+	ExternalObjectAddedObserverHandle = Database->RegisterQuery(
+	Select(
+		TEXT("Object Reinstancing: Added ExternalObjectColumn"),
+		FObserver::OnAdd<FTypedElementExternalObjectColumn>(),
+		[this](TypedElementDataStorage::IQueryContext& Context, TypedElementRowHandle Row, const FTypedElementExternalObjectColumn& ExternalObjectColumn)
+		{
+			TypedElementRowHandle Memento;
+			if (NewInstanceToMementoMap.RemoveAndCopyValue(ExternalObjectColumn.Object, Memento))
+			{
+				Context.AddColumn<FTypedElementMementoReinstanceTarget>(Memento, FTypedElementMementoReinstanceTarget{.Target = Row});
+			}
+		})
+		.Compile());
+}
+
+void UTypedElementObjectReinstancingManager::UnregisterQueries()
+{
+	Database->UnregisterQuery(ExternalObjectAddedObserverHandle);
+	Database->UnregisterQuery(UObjectAddedObserverHandle);
+}
+
+void UTypedElementObjectReinstancingManager::HandleOnObjectsReinstanced(
+	const FCoreUObjectDelegates::FReplacementObjectMap& ObjectReplacementMap)
+{
+	NewInstanceToMementoMap.Empty(ObjectReplacementMap.Num());
+	
+	for (FCoreUObjectDelegates::FReplacementObjectMap::TConstIterator Iter = ObjectReplacementMap.CreateConstIterator(); Iter; ++Iter)
+	{
+		const UObject* PreDeleteObject = Iter->Key;
+		const UObject* NewInstanceObject = Iter->Value;
+		
+		TypedElementRowHandle ObjectRow = DataStorageCompatibility->FindRowWithCompatibleObject(PreDeleteObject);
+		if (ObjectRow != TypedElementInvalidRowHandle)
+		{
+			TypedElementRowHandle Memento = MementoSystem->CreateMemento(Database.Get());
+			MementoSystem->EnableMementoizeOnDelete(Database.Get(), ObjectRow, Memento);
+
+			// Keep a temporary map of the new objects that will be created to the mementos that keep their state
+			NewInstanceToMementoMap.Add(NewInstanceObject, Memento);
+		}
+	}
+}
diff --git a/Engine/Plugins/Experimental/TypedElementsDataStorage/Source/TypedElementsDataStorage/Private/Compatibility/TypedElementObjectReinstancingManager.h b/Engine/Plugins/Experimental/TypedElementsDataStorage/Source/TypedElementsDataStorage/Private/Compatibility/TypedElementObjectReinstancingManager.h
new file mode 100644
--- /dev/null
+++ b/Engine/Plugins/Experimental/TypedElementsDataStorage/Source/TypedElementsDataStorage/Private/Compatibility/TypedElementObjectReinstancingManager.h
@@ -0,0 +1,44 @@
+﻿// Copyright Epic Games, Inc. All Rights Reserved.
+
+#pragma once
+
+#include "Elements/Interfaces/TypedElementDataStorageInterface.h"
+#include "Elements/Interfaces/TypedElementDataStorageCompatibilityInterface.h"
+#include "Memento/TypedElementMementoSystem.h"
+#include "UObject/UObjectGlobals.h"
+
+#include "TypedElementObjectReinstancingManager.generated.h"
+
+class ITypedElementDataStorageCompatibilityInterface;
+class ITypedElementDataStorageInterface;
+class UTypedElementDatabaseCompatibility;
+class UTypedElementMementoSystem;
+
+UCLASS(Transient)
+class UTypedElementObjectReinstancingManager : public UObject
+{
+	GENERATED_BODY()
+public:
+	UTypedElementObjectReinstancingManager();
+
+	void Initialize(UTypedElementDatabase& InDatabase, UTypedElementDatabaseCompatibility& InDataStorageCompatibility, UTypedElementMementoSystem& InMementoSystem);
+	void Deinitialize();
+
+private:
+	void RegisterQueries();
+	void UnregisterQueries();
+	void HandleOnObjectsReinstanced(const FCoreUObjectDelegates::FReplacementObjectMap& ObjectReplacementMap);
+
+	TObjectPtr<UTypedElementDatabase> Database = nullptr;
+	TObjectPtr<UTypedElementDatabaseCompatibility> DataStorageCompatibility = nullptr;
+	TObjectPtr<UTypedElementMementoSystem> MementoSystem = nullptr;
+	
+	// Maps objects that are about to be created to the memento that will be used to reinstance them
+	TMap<const void*, TypedElementRowHandle> NewInstanceToMementoMap;
+
+	TypedElementQueryHandle UObjectAddedObserverHandle;
+	TypedElementQueryHandle ExternalObjectAddedObserverHandle;
+	
+	TypedElementTableHandle MementoRowBaseTable;
+	FDelegateHandle ReinstancingCallbackHandle;
+};
diff --git a/Engine/Plugins/Experimental/TypedElementsDataStorage/Source/TypedElementsDataStorage/Private/Memento/TypedElementMementoRowTypes.h b/Engine/Plugins/Experimental/TypedElementsDataStorage/Source/TypedElementsDataStorage/Private/Memento/TypedElementMementoRowTypes.h
--- a/Engine/Plugins/Experimental/TypedElementsDataStorage/Source/TypedElementsDataStorage/Private/Memento/TypedElementMementoRowTypes.h
+++ b/Engine/Plugins/Experimental/TypedElementsDataStorage/Source/TypedElementsDataStorage/Private/Memento/TypedElementMementoRowTypes.h
@@ -20,19 +20,12 @@ struct FTypedElementMementoTag : public FTypedElementDataStorageTag
 };
 
 /**
- * Column for a memento row indicating that the row is populated with memento data
+ * A column added to a memento row which will trigger the process
+ * of converting memento columns to the matching mementoizable columns in
+ * the target row
  */
 USTRUCT()
-struct FTypedElementMementoPopulated : public FTypedElementDataStorageTag
-{
-	GENERATED_BODY()
-};
-
-/**
- * A column added to a memento row which will trigger the reinstantiation process
- */
-USTRUCT()
-struct FTypedElementReinstanceTarget : public FTypedElementDataStorageColumn
+struct FTypedElementMementoReinstanceTarget : public FTypedElementDataStorageColumn
 {
 	GENERATED_BODY()
 	TypedElementRowHandle Target;
diff --git a/Engine/Plugins/Experimental/TypedElementsDataStorage/Source/TypedElementsDataStorage/Private/Memento/TypedElementMementoSystem.cpp b/Engine/Plugins/Experimental/TypedElementsDataStorage/Source/TypedElementsDataStorage/Private/Memento/TypedElementMementoSystem.cpp
--- a/Engine/Plugins/Experimental/TypedElementsDataStorage/Source/TypedElementsDataStorage/Private/Memento/TypedElementMementoSystem.cpp
+++ b/Engine/Plugins/Experimental/TypedElementsDataStorage/Source/TypedElementsDataStorage/Private/Memento/TypedElementMementoSystem.cpp
@@ -3,137 +3,41 @@
 #include "TypedElementMementoSystem.h"
 
 #include "Elements/Framework/TypedElementQueryBuilder.h"
-#include "HAL/IConsoleManager.h"
 #include "Memento/TypedElementMementoInterface.h"
 #include "Memento/TypedElementMementoTranslators.h"
 #include "TypedElementMementoRowTypes.h"
+#include "TypedElementDatabase.h"
 
-DECLARE_LOG_CATEGORY_CLASS(LogTypedElementMemento, Log, All)
-
-namespace Private
-{
-	// Number of frames to keep memento rows before deletion
-	bool GMementosEnabled = false;
-	FAutoConsoleVariableRef CVarMementoEnable(
-		TEXT("teds.mementos.enable"),
-		GMementosEnabled,
-		TEXT("Enable memento system for newly added objects\n"));
-}
-
-void UTypedElementMementoSystemFactory::RegisterTables(ITypedElementDataStorageInterface& DataStorage) const
-{
-	MementoRowBaseTable = DataStorage.RegisterTable(
-		{FTypedElementMementoTag::StaticStruct()},
-		TEXT("MementoRowBaseTable"));
-
-	// Memento rows become populated when the FTypedElementMementoPopulated column is added
-	TypedElementTableHandle MementoRowBasePopulatedTable = DataStorage.RegisterTable(
-		MementoRowBaseTable,
-		{FTypedElementMementoPopulated::StaticStruct()},
-		TEXT("MementoRowBasePopulatedTable"));
-}
-
-void UTypedElementMementoSystemFactory::RegisterRegistrationFilters(
-	ITypedElementDataStorageCompatibilityInterface& DataStorageCompatibility) const
+void UTypedElementMementoSystem::Initialize(UTypedElementDatabase& DataStorage)
 {
-	const_cast<UTypedElementMementoSystemFactory*>(this)->RegisterWithCompatibilityLayer(DataStorageCompatibility);
+	RegisterTables(DataStorage);
+	RegisterQueries(DataStorage);
 }
 
-void UTypedElementMementoSystemFactory::RegisterWithCompatibilityLayer(ITypedElementDataStorageCompatibilityInterface& DataStorageCompatibility)
+void UTypedElementMementoSystem::Deinitialize()
 {
-	if (Private::GMementosEnabled)
-	{
-		PRAGMA_DISABLE_DEPRECATION_WARNINGS
-		ObjectAddedDelegateHandle = DataStorageCompatibility.GetOnObjectAddedDelegate().AddUObject(this, &UTypedElementMementoSystemFactory::HandleObjectAddedToCompatibility);
-		ObjectRemovedDelegateHandle = DataStorageCompatibility.GetOnObjectPreDestroy().AddUObject(this, &UTypedElementMementoSystemFactory::HandleObjectPreRemoveFromCompatibility);
-		PRAGMA_ENABLE_DEPRECATION_WARNINGS
-	}
-	
-	auto HandleCVarEnabledChanged = [this, &DataStorageCompatibility](IConsoleVariable* CVar)
-	{
-		// Toggled Off
-		if (!Private::GMementosEnabled)
-		{
-			if (ObjectAddedDelegateHandle.IsValid())
-			{
-				PRAGMA_DISABLE_DEPRECATION_WARNINGS
-				DataStorageCompatibility.GetOnObjectAddedDelegate().Remove(ObjectAddedDelegateHandle);
-				DataStorageCompatibility.GetOnObjectPreDestroy().Remove(ObjectRemovedDelegateHandle);
-				PRAGMA_ENABLE_DEPRECATION_WARNINGS
-				ObjectAddedDelegateHandle.Reset();
-			}
-		}
-		// Toggled on
-		else if (Private::GMementosEnabled)
-		{
-			if (!ObjectAddedDelegateHandle.IsValid())
-			{
-				PRAGMA_DISABLE_DEPRECATION_WARNINGS
-				ObjectAddedDelegateHandle = DataStorageCompatibility.GetOnObjectAddedDelegate().AddUObject(this, &UTypedElementMementoSystemFactory::HandleObjectAddedToCompatibility);
-				ObjectRemovedDelegateHandle = DataStorageCompatibility.GetOnObjectPreDestroy().AddUObject(this, &UTypedElementMementoSystemFactory::HandleObjectPreRemoveFromCompatibility);
-				PRAGMA_ENABLE_DEPRECATION_WARNINGS
-			}
-		}
-	};
-	
-	Private::CVarMementoEnable->OnChangedDelegate().AddLambda(HandleCVarEnabledChanged);
-	
-	FCoreUObjectDelegates::OnObjectsReinstanced.AddUObject(this, &UTypedElementMementoSystemFactory::HandleOnObjectsReinstanced);
 }
 
-/**
- * Object reinstancing happens in the following sequence:
- * - This HandleOnObjectsReinstanced callback is fired
- * - The old instance objects are deleted
- * - The new instance objects are added
- */
-void UTypedElementMementoSystemFactory::HandleOnObjectsReinstanced(
-	const FCoreUObjectDelegates::FReplacementObjectMap& ObjectReplacementMap)
+TypedElementRowHandle UTypedElementMementoSystem::CreateMemento(ITypedElementDataStorageInterface* DataStorage)
 {
-	for (FCoreUObjectDelegates::FReplacementObjectMap::TConstIterator Iter = ObjectReplacementMap.CreateConstIterator(); Iter; ++Iter)
-	{
-		const UObject* PreDeleteObject = Iter->Key;
-		const UObject* NewInstanceObject = Iter->Value;
-		
-		const TypedElementRowHandle* MementoRowPtr = MementoizableObjects.Find(PreDeleteObject);
-		if (MementoRowPtr != nullptr)
-		{
-			NewInstanceToMementoMap.Add(NewInstanceObject, *MementoRowPtr);
-			UE_LOG(LogTypedElementMemento, VeryVerbose, TEXT("Instance Replacement: 0x%p -> %llu [%s]"), NewInstanceObject, *MementoRowPtr, *NewInstanceObject->GetClass()->GetName());
-		}
-	}
+	return DataStorage->AddRow(MementoRowBaseTable);
 }
 
-void UTypedElementMementoSystemFactory::HandleObjectAddedToCompatibility(ITypedElementDataStorageInterface* Storage, const void* Object, const FTypedElementDatabaseCompatibilityObjectTypeInfo& TypeInfo, TypedElementRowHandle Row)
+void UTypedElementMementoSystem::EnableMementoizeOnDelete(ITypedElementDataStorageInterface* DataStorage, TypedElementRowHandle Row, TypedElementRowHandle Memento)
 {
-	// Register row unconditionally for mementoization
-	{
-		TypedElementRowHandle Memento = Storage->AddRow(MementoRowBaseTable);
-		Storage->AddOrGetColumn<FTypedElementMementoOnDelete>(Row, FTypedElementMementoOnDelete{ .Memento = Memento });
-		UE_LOG(LogTypedElementMemento, VeryVerbose, TEXT("Object [%s] tagged for mementoization 0x%p -> %llu"), *TypeInfo.GetFName().ToString(), Object, Memento);
-		MementoizableObjects.Add(Object, Memento);
-	}
-
-	// If this object is a reinstantiation of a deleted object, then setup the memento row to with the target
-	// row to trigger reinstantiation
-	TypedElementRowHandle Memento;
-	if (NewInstanceToMementoMap.RemoveAndCopyValue(Object, Memento))
-	{
-		UE_LOG(LogTypedElementMemento, VeryVerbose, TEXT("New instance object [%s] mapped to memento 0x%p -> %llu"), *TypeInfo.GetFName().ToString(), Object, Memento);
-		Storage->AddOrGetColumn<FTypedElementReinstanceTarget>(Memento, FTypedElementReinstanceTarget{.Target = Row});
-	}
+	DataStorage->AddOrGetColumn<FTypedElementMementoOnDelete>(Row, FTypedElementMementoOnDelete{ .Memento = Memento });
 }
 
-void UTypedElementMementoSystemFactory::HandleObjectPreRemoveFromCompatibility(ITypedElementDataStorageInterface* Storage,
-	const void* Object, const FTypedElementDatabaseCompatibilityObjectTypeInfo& TypeInfo, TypedElementRowHandle Row)
+void UTypedElementMementoSystem::RegisterTables(UTypedElementDatabase& DataStorage)
 {
-	UE_LOG(LogTypedElementMemento, VeryVerbose, TEXT("Removing object 0x%p"),Object);
-
-	// Remove row unconditionally from mementoization
-	MementoizableObjects.Remove(Object);
+	// Register tables that will be used by reinstancing
+	MementoRowBaseTable = DataStorage.RegisterTable(
+		{FTypedElementMementoTag::StaticStruct()},
+		TEXT("MementoRowBaseTable"));
+	
 }
 
-void UTypedElementMementoSystemFactory::RegisterQueries(ITypedElementDataStorageInterface& DataStorage) const
+void UTypedElementMementoSystem::RegisterQueries(UTypedElementDatabase& DataStorage) const
 {
 	using DSI = ITypedElementDataStorageInterface;
 	
@@ -185,7 +89,10 @@ void UTypedElementMementoSystemFactory::RegisterQueries(ITypedElementDataStorage
 		check(QueryHandle != TypedElementInvalidQueryHandle);
 	}
 
-	// Setup processors to execute memento translators on to reinstantiate rows
+	// Setup processors to execute memento translators on reinstantiate columns onto a target
+	// The primary query runs over all mementos with a reinstance target and a column holding memento data
+	// The subquery runs on rows targeted by the reinstance column that also have a mementoizable column
+	// Note: Currently it is not possible for mementos to create a new column from a subquery
 	for (int32 Index = 0, End = MementoTranslators.Num(); Index < End; ++Index)
 	{
 		const UTypedElementMementoTranslatorBase* MementoTranslator = MementoTranslators[Index];
@@ -202,7 +109,7 @@ void UTypedElementMementoSystemFactory::RegisterQueries(ITypedElementDataStorage
 			Select(
 				TranslationProcessorName,
 				FProcessor(DSI::EQueryTickPhase::PostPhysics, DataStorage.GetQueryTickGroupName(DSI::EQueryTickGroups::Default)),
-				[MementoTranslator](TypedElementDataStorage::IQueryContext& Context, TypedElementRowHandle Row, const FTypedElementReinstanceTarget& ReinstanceTarget)
+				[MementoTranslator](TypedElementDataStorage::IQueryContext& Context, TypedElementRowHandle Row, const FTypedElementMementoReinstanceTarget& ReinstanceTarget)
 				{
 					const UScriptStruct* MementoType = MementoTranslator->GetMementoType();
 					
@@ -224,37 +131,20 @@ void UTypedElementMementoSystemFactory::RegisterQueries(ITypedElementDataStorage
 		check(QueryHandle != TypedElementInvalidQueryHandle);
 	}
 
-	// Setup a default policy of deleting populated mementos after a set number of frames
-	// The first use-case of mementos are to support reinstancing which occurs over a single frame
-	// Therefore it should be safe to remove populated mementos after a couple of frames since the data
-	// is no longer needed.
-	// Expansion of memento deletion policies will be necessary for other use cases
-	{
-		const TypedElementQueryHandle QueryHandle = DataStorage.RegisterQuery(
-			Select(
-				TEXT("Add Memento Populated Tag"),
-				FObserver::OnRemove<FTypedElementMementoOnDelete>(), // When a memento is populated
-				[](TypedElementDataStorage::IQueryContext& Context, TypedElementRowHandle Row, const FTypedElementMementoOnDelete& MementoRow)
-				{
-					Context.AddColumn<FTypedElementMementoPopulated>(MementoRow.Memento, FTypedElementMementoPopulated{});
-				})
-				.Compile());
-		check(QueryHandle != TypedElementInvalidQueryHandle);
-	}
-
 	/**
-	 * A processor which deletes mementos that were populated more than N frames ago
+	 * A processor which deletes mementos that were used for rehydrating another row
 	 */
 	{
 		const TypedElementQueryHandle QueryHandle = DataStorage.RegisterQuery(
 			Select(
-			TEXT("Delete populated mementos"),
+			TEXT("Delete mementos used for reinstancing"),
 			FProcessor(DSI::EQueryTickPhase::FrameEnd, DataStorage.GetQueryTickGroupName(DSI::EQueryTickGroups::Default)),
 				[](TypedElementDataStorage::IQueryContext& Context, TypedElementRowHandle Row)
 				{
 					Context.RemoveRow(Row);
 				})
-				.Where().All<FTypedElementMementoTag, FTypedElementMementoPopulated>()
+				.ReadOnly<FTypedElementMementoReinstanceTarget>()
+				.Where().All<FTypedElementMementoTag>()
 				.Compile()
 			);
 		check(QueryHandle != TypedElementInvalidQueryHandle);
diff --git a/Engine/Plugins/Experimental/TypedElementsDataStorage/Source/TypedElementsDataStorage/Private/Memento/TypedElementMementoSystem.h b/Engine/Plugins/Experimental/TypedElementsDataStorage/Source/TypedElementsDataStorage/Private/Memento/TypedElementMementoSystem.h
--- a/Engine/Plugins/Experimental/TypedElementsDataStorage/Source/TypedElementsDataStorage/Private/Memento/TypedElementMementoSystem.h
+++ b/Engine/Plugins/Experimental/TypedElementsDataStorage/Source/TypedElementsDataStorage/Private/Memento/TypedElementMementoSystem.h
@@ -8,35 +8,28 @@
 
 #include "TypedElementMementoSystem.generated.h"
 
-
-
 class ITypedElementDataStorageInterface;
 class IConsoleVariable;
 struct FTypedElementDatabaseCompatibilityObjectTypeInfo;
+class UTypedElementDatabase;
 
 UCLASS()
-class UTypedElementMementoSystemFactory : public UTypedElementDataStorageFactory
+class UTypedElementMementoSystem : public UObject
 {
 	GENERATED_BODY()
-
 public:
-	void RegisterQueries(ITypedElementDataStorageInterface& DataStorage) const override;
-	void RegisterTables(ITypedElementDataStorageInterface& DataStorage) const override;
-	void RegisterRegistrationFilters(ITypedElementDataStorageCompatibilityInterface& DataStorageCompatibility) const override;
-private:
 
-	void RegisterWithCompatibilityLayer(ITypedElementDataStorageCompatibilityInterface& DataStorageCompatibility);
-	void HandleObjectAddedToCompatibility(ITypedElementDataStorageInterface* Storage, const void* Object, const FTypedElementDatabaseCompatibilityObjectTypeInfo& TypeInfo, TypedElementRowHandle Row);
-	void HandleObjectPreRemoveFromCompatibility(ITypedElementDataStorageInterface* Storage, const void* Object, const FTypedElementDatabaseCompatibilityObjectTypeInfo& TypeInfo, TypedElementRowHandle Row);
-	void HandleOnObjectsReinstanced(const FCoreUObjectDelegates::FReplacementObjectMap& ObjectReplacementMap);
-	
-	mutable TypedElementTableHandle MementoRowBaseTable;
-	FDelegateHandle ObjectAddedDelegateHandle;
-	FDelegateHandle ObjectRemovedDelegateHandle;
+	void Initialize(UTypedElementDatabase& DataStorage);
+	void Deinitialize();
 
-	// Maps existing objects to their Memento
-	TMap<const void*, TypedElementRowHandle> MementoizableObjects;
+	TypedElementRowHandle CreateMemento(ITypedElementDataStorageInterface* DataStorage);
+	
+	// Enables the given Row to be mementoized into the Memento row when it is deleted
+	void EnableMementoizeOnDelete(ITypedElementDataStorageInterface* DataStorage, TypedElementRowHandle Row, TypedElementRowHandle Memento);
 
-	// Maps objects that are about to be created to the memento that will be used to reinstance them
-	TMap<const void*, TypedElementRowHandle> NewInstanceToMementoMap;
+private:
+	void RegisterQueries(UTypedElementDatabase& DataStorage) const;
+	void RegisterTables(UTypedElementDatabase& DataStorage);
+	
+	TypedElementTableHandle MementoRowBaseTable;
 };
diff --git a/Engine/Plugins/Experimental/TypedElementsDataStorage/Source/TypedElementsDataStorage/Private/TypedElementDatabaseCompatibility.cpp b/Engine/Plugins/Experimental/TypedElementsDataStorage/Source/TypedElementsDataStorage/Private/TypedElementDatabaseCompatibility.cpp
--- a/Engine/Plugins/Experimental/TypedElementsDataStorage/Source/TypedElementsDataStorage/Private/TypedElementDatabaseCompatibility.cpp
+++ b/Engine/Plugins/Experimental/TypedElementsDataStorage/Source/TypedElementsDataStorage/Private/TypedElementDatabaseCompatibility.cpp
@@ -133,17 +133,6 @@ void UTypedElementDatabaseCompatibility::RemoveCompatibleObjectExplicit(UObject*
 		TypedElementRowHandle Row;
 		if (ReverseObjectLookup.RemoveAndCopyValue(Object, Row))
 		{
-			if (OnObjectPreDestroy.IsBound() && Storage->IsRowAvailable(Row))
-			{
-				// This will be null if object has not yet been added (ie. is in UObjectsPendingRegistration)
-				// Don't call it as OnAdded hasn't be called yet
-				const FTypedElementClassTypeInfoColumn* TypeInfo = Storage->GetColumn<FTypedElementClassTypeInfoColumn>(Row);
-				if (TypeInfo != nullptr)
-				{
-					OnObjectPreDestroy.Broadcast(Storage, Object, FTypedElementDatabaseCompatibilityObjectTypeInfo(TypeInfo->TypeInfo.Get()), Row);
-				}
-			}
-			
 			Storage->RemoveRow(Row);
 		}
 	}
@@ -154,18 +143,7 @@ void UTypedElementDatabaseCompatibility::RemoveCompatibleObjectExplicit(void* Ob
 	checkf(Storage, TEXT("Removing compatible objects is not supported before Typed Element's Database compatibility manager has been initialized."));
 	TypedElementRowHandle Row;
 	if (ReverseObjectLookup.RemoveAndCopyValue(Object, Row))
-	{
-		if (OnObjectPreDestroy.IsBound() && Storage->IsRowAvailable(Row))
-		{
-			// This will be null if object has not yet been added (ie. is in ExternalObjectsPendingRegistration)
-			// Don't call it as OnAdded hasn't be called yet
-			const FTypedElementScriptStructTypeInfoColumn* TypeInfo = Storage->GetColumn<FTypedElementScriptStructTypeInfoColumn>(Row);
-			if (TypeInfo != nullptr)
-			{
-				OnObjectPreDestroy.Broadcast(Storage, Object, FTypedElementDatabaseCompatibilityObjectTypeInfo(TypeInfo->TypeInfo.Get()), Row);
-			}
-		}
-		
+	{		
 		Storage->RemoveRow(Row);
 	}
 }
@@ -189,16 +167,6 @@ void UTypedElementDatabaseCompatibility::RemoveCompatibleObjectExplicit(AActor*
 			if (ActorStore && !ActorStore->IsOwnedByMass()) // Only remove actors that were externally created.
 			{
 				TypedElementRowHandle Row = Entity.AsNumber();
-				if (OnObjectPreDestroy.IsBound() && Storage->IsRowAvailable(Entity.AsNumber()))
-				{
-					// This will be null if object has not yet been added (ie. is in ActorsPendingRegistration)
-					// Don't call it as OnAdded hasn't be called yet
-					const FTypedElementClassTypeInfoColumn* TypeInfo = Storage->GetColumn<FTypedElementClassTypeInfoColumn>(Row);
-					if (TypeInfo != nullptr)
-					{
-						OnObjectPreDestroy.Broadcast(Storage, Actor, FTypedElementDatabaseCompatibilityObjectTypeInfo(TypeInfo->TypeInfo.Get()), Row);
-					}
-				}
 				
 				ActorSubsystem->RemoveHandleForActor(Actor);
 				Storage->RemoveRow(Entity.AsNumber());
@@ -251,16 +219,6 @@ TypedElementRowHandle UTypedElementDatabaseCompatibility::FindRowWithCompatibleO
 	return TypedElementInvalidRowHandle;
 }
 
-FTypedElementDatabaseCompatibility_OnObjectAdded& UTypedElementDatabaseCompatibility::GetOnObjectAddedDelegate()
-{
-	return OnObjectAddedDelegate;
-}
-
-FTypedElementDatabaseCompatibility_OnObjectAdded& UTypedElementDatabaseCompatibility::GetOnObjectPreDestroy()
-{
-	return OnObjectPreDestroy;
-}
-
 void UTypedElementDatabaseCompatibility::Prepare()
 {
 	UMassActorEditorSubsystem* MassActorEditorSubsystem = Storage->GetExternalSystem<UMassActorEditorSubsystem>();
@@ -468,8 +426,6 @@ void UTypedElementDatabaseCompatibility::TickPendingActorRegistration(UWorld* Ed
 					
 					// Make sure the new row is tagged for update.
 					Storage->AddColumn<FTypedElementSyncFromWorldTag>(Row);
-
-					OnObjectAddedDelegate.Broadcast(Storage, Actor, FTypedElementDatabaseCompatibilityObjectTypeInfo(Actor->GetClass()) , Row);
 				});
 		}
 			
@@ -502,8 +458,6 @@ void UTypedElementDatabaseCompatibility::TickPendingUObjectRegistration()
 					Storage->AddOrGetColumn<FTypedElementClassTypeInfoColumn>(Row, FTypedElementClassTypeInfoColumn{ .TypeInfo = Object->GetClass() });
 					// Make sure the new row is tagged for update.
 					Storage->AddColumn<FTypedElementSyncFromWorldTag>(Row);
-
-					OnObjectAddedDelegate.Broadcast(Storage, Object.Get(), FTypedElementDatabaseCompatibilityObjectTypeInfo(Object->GetClass()), Row);
 				});
 		}
 
@@ -535,8 +489,6 @@ void UTypedElementDatabaseCompatibility::TickPendingExternalObjectRegistration()
 					Storage->AddOrGetColumn<FTypedElementScriptStructTypeInfoColumn>(Row, FTypedElementScriptStructTypeInfoColumn{ .TypeInfo = Object.TypeInfo });
 					// Make sure the new row is tagged for update.
 					Storage->AddColumn<FTypedElementSyncFromWorldTag>(Row);
-
-					OnObjectAddedDelegate.Broadcast(Storage, Object.Object, FTypedElementDatabaseCompatibilityObjectTypeInfo(Object.TypeInfo.Get()), Row);
 				});
 		}
 
diff --git a/Engine/Plugins/Experimental/TypedElementsDataStorage/Source/TypedElementsDataStorage/Private/TypedElementDatabaseCompatibility.h b/Engine/Plugins/Experimental/TypedElementsDataStorage/Source/TypedElementsDataStorage/Private/TypedElementDatabaseCompatibility.h
--- a/Engine/Plugins/Experimental/TypedElementsDataStorage/Source/TypedElementsDataStorage/Private/TypedElementDatabaseCompatibility.h
+++ b/Engine/Plugins/Experimental/TypedElementsDataStorage/Source/TypedElementsDataStorage/Private/TypedElementDatabaseCompatibility.h
@@ -5,6 +5,7 @@
 #include "Containers/Array.h"
 #include "Elements/Interfaces/TypedElementDataStorageInterface.h"
 #include "Elements/Interfaces/TypedElementDataStorageCompatibilityInterface.h"
+#include "Compatibility/TypedElementObjectReinstancingManager.h"
 #include "UObject/ObjectMacros.h"
 #include "UObject/WeakObjectPtr.h"
 
@@ -43,9 +44,6 @@ public:
 	TypedElementRowHandle FindRowWithCompatibleObjectExplicit(const UObject* Object) const override;
 	TypedElementRowHandle FindRowWithCompatibleObjectExplicit(const AActor* Actor) const override;
 	TypedElementRowHandle FindRowWithCompatibleObjectExplicit(const void* Object) const override;
-	
-	FTypedElementDatabaseCompatibility_OnObjectAdded& GetOnObjectAddedDelegate() override;
-	FTypedElementDatabaseCompatibility_OnObjectPreDestroy& GetOnObjectPreDestroy() override;
 
 private:
 	void Prepare();
@@ -94,8 +92,6 @@ private:
 	
 	TArray<ObjectRegistrationFilter> ObjectRegistrationFilters;
 	TArray<ObjectToRowDealiaser> ObjectToRowDialiasers;
-	FTypedElementDatabaseCompatibility_OnObjectAdded OnObjectAddedDelegate;
-	FTypedElementDatabaseCompatibility_OnObjectPreDestroy OnObjectPreDestroy;
 
 	TypedElementTableHandle StandardActorTable{ TypedElementInvalidTableHandle };
 	TypedElementTableHandle StandardActorWithTransformTable{ TypedElementInvalidTableHandle };
diff --git a/Engine/Plugins/Experimental/TypedElementsDataStorage/Source/TypedElementsDataStorage/Private/TypedElementsDataStorage.cpp b/Engine/Plugins/Experimental/TypedElementsDataStorage/Source/TypedElementsDataStorage/Private/TypedElementsDataStorage.cpp
--- a/Engine/Plugins/Experimental/TypedElementsDataStorage/Source/TypedElementsDataStorage/Private/TypedElementsDataStorage.cpp
+++ b/Engine/Plugins/Experimental/TypedElementsDataStorage/Source/TypedElementsDataStorage/Private/TypedElementsDataStorage.cpp
@@ -64,6 +64,12 @@ void FTypedElementsDataStorageModule::StartupModule()
 				DatabaseUi = NewObject<UTypedElementDatabaseUi>();
 				DatabaseUi->Initialize(Database.Get(), DatabaseCompatibility.Get());
 
+				MementoSystem = NewObject<UTypedElementMementoSystem>();
+				MementoSystem->Initialize(*Database.Get());
+
+				ObjectReinstancingManager = NewObject<UTypedElementObjectReinstancingManager>();
+				ObjectReinstancingManager->Initialize(*Database, *DatabaseCompatibility, *MementoSystem);
+
 				// Register the various database instances.
 				UTypedElementRegistry* Registry = UTypedElementRegistry::GetInstance();
 				checkf(Registry, TEXT(
@@ -128,6 +134,8 @@ void FTypedElementsDataStorageModule::ShutdownModule()
 
 		if (UObjectInitialized())
 		{
+			ObjectReinstancingManager->Deinitialize();
+			MementoSystem->Deinitialize();
 			DatabaseUi->Deinitialize();
 			DatabaseCompatibility->Deinitialize();
 			Database->Deinitialize();
@@ -144,6 +152,8 @@ void FTypedElementsDataStorageModule::AddReferencedObjects(FReferenceCollector&
 		Collector.AddReferencedObject(Database);
 		Collector.AddReferencedObject(DatabaseCompatibility);
 		Collector.AddReferencedObject(DatabaseUi);
+		Collector.AddReferencedObject(MementoSystem);
+		Collector.AddReferencedObject(ObjectReinstancingManager);
 	}
 }
 
diff --git a/Engine/Plugins/Experimental/TypedElementsDataStorage/Source/TypedElementsDataStorage/Public/TypedElementsDataStorage.h b/Engine/Plugins/Experimental/TypedElementsDataStorage/Source/TypedElementsDataStorage/Public/TypedElementsDataStorage.h
--- a/Engine/Plugins/Experimental/TypedElementsDataStorage/Source/TypedElementsDataStorage/Public/TypedElementsDataStorage.h
+++ b/Engine/Plugins/Experimental/TypedElementsDataStorage/Source/TypedElementsDataStorage/Public/TypedElementsDataStorage.h
@@ -11,6 +11,8 @@ class FReferenceCollector;
 class UTypedElementDatabase;
 class UTypedElementDatabaseCompatibility;
 class UTypedElementDatabaseUi;
+class UTypedElementObjectReinstancingManager;
+class UTypedElementMementoSystem;
 
 class FTypedElementsDataStorageModule : public IModuleInterface, public FGCObject
 {
@@ -36,5 +38,7 @@ private:
 	TObjectPtr<UTypedElementDatabase> Database;
 	TObjectPtr<UTypedElementDatabaseCompatibility> DatabaseCompatibility;
 	TObjectPtr<UTypedElementDatabaseUi> DatabaseUi;
+	TObjectPtr<UTypedElementMementoSystem> MementoSystem;
+	TObjectPtr<UTypedElementObjectReinstancingManager> ObjectReinstancingManager;
 	bool bInitialized{ false };
 };
diff --git a/Engine/Source/Runtime/TypedElementFramework/Public/Elements/Interfaces/TypedElementDataStorageCompatibilityInterface.h b/Engine/Source/Runtime/TypedElementFramework/Public/Elements/Interfaces/TypedElementDataStorageCompatibilityInterface.h
--- a/Engine/Source/Runtime/TypedElementFramework/Public/Elements/Interfaces/TypedElementDataStorageCompatibilityInterface.h
+++ b/Engine/Source/Runtime/TypedElementFramework/Public/Elements/Interfaces/TypedElementDataStorageCompatibilityInterface.h
@@ -16,20 +16,6 @@
 class AActor;
 struct FTypedElementDatabaseCompatibilityObjectTypeInfo;
 
-DECLARE_MULTICAST_DELEGATE_FourParams(
-	FTypedElementDatabaseCompatibility_OnObjectAdded,
-	ITypedElementDataStorageInterface* Storage,
-	const void* Object,
-	const FTypedElementDatabaseCompatibilityObjectTypeInfo& ObjectTypeInfo,
-	TypedElementRowHandle Row);
-
-DECLARE_MULTICAST_DELEGATE_FourParams(
-	FTypedElementDatabaseCompatibility_OnObjectPreDestroy,
-	ITypedElementDataStorageInterface* Storage,
-	const void* Object,
-	const FTypedElementDatabaseCompatibilityObjectTypeInfo& ObjectTypeInfo,
-	TypedElementRowHandle Row);
-
 UINTERFACE(MinimalAPI)
 class UTypedElementDataStorageCompatibilityInterface : public UInterface
 {
@@ -123,18 +109,6 @@ public:
 	virtual TypedElementRowHandle FindRowWithCompatibleObjectExplicit(const AActor* Actor) const = 0;
 	/** Finds a previously stored FStruct. If not found an invalid row handle will be returned. */
 	virtual TypedElementRowHandle FindRowWithCompatibleObjectExplicit(const void* Object) const = 0;
-	
-	/**
-	 * An extension point called when objects have been added to the compatibility interface
-	 */
-	UE_DEPRECATED(5.4, "This is temporary for adding memento support. Do no rely on it.")
-	virtual FTypedElementDatabaseCompatibility_OnObjectAdded& GetOnObjectAddedDelegate() = 0;
-	
-	/**
-	 * An extension point called when objects have been removed from the compatibility interface
-	 */
-	UE_DEPRECATED(5.4, "This is temporary for adding memento support. Do no rely on it.")
-	virtual FTypedElementDatabaseCompatibility_OnObjectPreDestroy& GetOnObjectPreDestroy() = 0;
 };
 
 template<typename Type> Type* GetRawPointer(const TWeakObjectPtr<Type> Object)	{ return Object.Get(); }
@@ -251,49 +225,4 @@ struct TTypedElementSubsystemTraits final
 			return false;
 		}
 	}
-};
-
-enum class ETypedElementDatabaseCompatibilityObjectType
-{
-	Struct,
-	Class
-};
-
-/**
- * Objects with type info defined in either UScriptStruct or UClass can be stored into TEDS via the
- * ITypedElementDataStorageCompatibilityInterface
- * This is a discriminated union which aids with callbacks made when objects are added
- */
-struct FTypedElementDatabaseCompatibilityObjectTypeInfo
-{
-	ETypedElementDatabaseCompatibilityObjectType TypeInfoType;
-
-	union
-	{
-		const UScriptStruct* ScriptStruct;
-		const UClass* Class;
-	};
-
-	explicit FTypedElementDatabaseCompatibilityObjectTypeInfo(const UScriptStruct* InScriptStruct)
-		: TypeInfoType(ETypedElementDatabaseCompatibilityObjectType::Struct)
-		, ScriptStruct(InScriptStruct)
-	{}
-
-	explicit FTypedElementDatabaseCompatibilityObjectTypeInfo(const UClass* InClass)
-	: TypeInfoType(ETypedElementDatabaseCompatibilityObjectType::Class)
-	, Class(InClass)
-	{}
-
-	FName GetFName() const;
-};
-
-
-inline FName FTypedElementDatabaseCompatibilityObjectTypeInfo::GetFName() const
-{
-	switch(TypeInfoType)
-	{
-	case ETypedElementDatabaseCompatibilityObjectType::Struct: return ScriptStruct->GetFName();
-	case ETypedElementDatabaseCompatibilityObjectType::Class: return Class->GetFName();
-	default: return FName();
-	}
-}
+};
\ No newline at end of file
