Disable deregistering ObjectReinstancingManager queries
* This operation is not supported by Mass at this time.

diff --git a/Engine/Plugins/Experimental/TypedElementsDataStorage/Source/TypedElementsDataStorage/Private/Compatibility/TypedElementObjectReinstancingManager.cpp b/Engine/Plugins/Experimental/TypedElementsDataStorage/Source/TypedElementsDataStorage/Private/Compatibility/TypedElementObjectReinstancingManager.cpp
--- a/Engine/Plugins/Experimental/TypedElementsDataStorage/Source/TypedElementsDataStorage/Private/Compatibility/TypedElementObjectReinstancingManager.cpp
+++ b/Engine/Plugins/Experimental/TypedElementsDataStorage/Source/TypedElementsDataStorage/Private/Compatibility/TypedElementObjectReinstancingManager.cpp
@@ -72,8 +72,9 @@ void UTypedElementObjectReinstancingManager::RegisterQueries()
 
 void UTypedElementObjectReinstancingManager::UnregisterQueries()
 {
-	Database->UnregisterQuery(ExternalObjectAddedObserverHandle);
-	Database->UnregisterQuery(UObjectAddedObserverHandle);
+	// TODO: Observer queries cannot be unregistered in Mass
+	// Database->UnregisterQuery(ExternalObjectAddedObserverHandle);
+	// Database->UnregisterQuery(UObjectAddedObserverHandle);
 }
 
 void UTypedElementObjectReinstancingManager::HandleOnObjectsReinstanced(
