Integration of the new boolean query conditions into TEDS UI.

This change integrates the new boolean query condition into TEDS UI so widget constructors can be registered to be matched with columns in a wider set of configurations. This also adds a templated utility function to make it less verbose to register widget constructors.

diff --git a/Engine/Plugins/Experimental/TypedElementsDataStorage/Source/TypedElementsDataStorage/Private/TypedElementDatabaseUI.cpp b/Engine/Plugins/Experimental/TypedElementsDataStorage/Source/TypedElementsDataStorage/Private/TypedElementDatabaseUI.cpp
--- a/Engine/Plugins/Experimental/TypedElementsDataStorage/Source/TypedElementsDataStorage/Private/TypedElementDatabaseUI.cpp
+++ b/Engine/Plugins/Experimental/TypedElementsDataStorage/Source/TypedElementsDataStorage/Private/TypedElementDatabaseUI.cpp
@@ -2,6 +2,9 @@
 
 #include "TypedElementDatabaseUI.h"
 
+#include "Algo/BinarySearch.h"
+#include "Algo/Sort.h"
+#include "Algo/Unique.h"
 #include "Elements/Columns/TypedElementSlateWidgetColumns.h"
 #include "Elements/Interfaces/TypedElementDataStorageInterface.h"
 #include "Elements/Interfaces/TypedElementDataStorageCompatibilityInterface.h"
@@ -92,7 +95,7 @@ bool UTypedElementDatabaseUi::RegisterWidgetFactory(FName Purpose, const UScript
 }
 
 bool UTypedElementDatabaseUi::RegisterWidgetFactory(
-	FName Purpose, const UScriptStruct* Constructor, TArray<TWeakObjectPtr<const UScriptStruct>> Columns)
+	FName Purpose, const UScriptStruct* Constructor, TypedElementQueryBuilder::FQueryConditions Columns)
 {
 	if (!Columns.IsEmpty())
 	{
@@ -109,18 +112,25 @@ bool UTypedElementDatabaseUi::RegisterWidgetFactory(
 				PurposeInfo->bIsSorted = false;
 				return true;
 			case ITypedElementDataStorageUiInterface::EPurposeType::UniqueByName:
-				if (PurposeInfo->Factories.IsEmpty())
+				if (!Columns.IsEmpty())
 				{
-					PurposeInfo->Factories.Emplace(Constructor);
-					PurposeInfo->bIsSorted = false;
+					if (PurposeInfo->Factories.IsEmpty())
+					{
+						PurposeInfo->Factories.Emplace(Constructor);
+						PurposeInfo->bIsSorted = false;
+					}
+					else
+					{
+						PurposeInfo->Factories.EmplaceAt(0, Constructor);
+					}
+					return true;
 				}
 				else
 				{
-					PurposeInfo->Factories.EmplaceAt(0, Constructor);
+					return false;
 				}
-				return true;
 			case ITypedElementDataStorageUiInterface::EPurposeType::UniqueByNameAndColumn:
-				if (PrepareColumnsList(Columns))
+				if (!Columns.IsEmpty())
 				{
 					PurposeInfo->Factories.Emplace(Constructor, MoveTemp(Columns));
 					PurposeInfo->bIsSorted = false;
@@ -190,7 +200,7 @@ bool UTypedElementDatabaseUi::RegisterWidgetFactory(FName Purpose, TUniquePtr<FT
 }
 
 bool UTypedElementDatabaseUi::RegisterWidgetFactory(FName Purpose, TUniquePtr<FTypedElementWidgetConstructor>&& Constructor, 
-	TArray<TWeakObjectPtr<const UScriptStruct>> Columns)
+	TypedElementQueryBuilder::FQueryConditions Columns)
 {
 	if (!Columns.IsEmpty())
 	{
@@ -205,21 +215,34 @@ bool UTypedElementDatabaseUi::RegisterWidgetFactory(FName Purpose, TUniquePtr<FT
 				PurposeInfo->bIsSorted = false;
 				return true;
 			case ITypedElementDataStorageUiInterface::EPurposeType::UniqueByName:
-				if (PurposeInfo->Factories.IsEmpty())
+				if (!Columns.IsEmpty())
 				{
-					PurposeInfo->Factories.Emplace(MoveTemp(Constructor));
-					PurposeInfo->bIsSorted = false;
+					if (PurposeInfo->Factories.IsEmpty())
+					{
+						PurposeInfo->Factories.Emplace(MoveTemp(Constructor));
+						PurposeInfo->bIsSorted = false;
+					}
+					else
+					{
+						PurposeInfo->Factories.EmplaceAt(0, MoveTemp(Constructor));
+					}
+					return true;
 				}
 				else
 				{
-					PurposeInfo->Factories.EmplaceAt(0, MoveTemp(Constructor));
+					return false;
 				}
-				return true;
 			case ITypedElementDataStorageUiInterface::EPurposeType::UniqueByNameAndColumn:
-				PrepareColumnsList(Columns);
-				PurposeInfo->Factories.Emplace(MoveTemp(Constructor), MoveTemp(Columns));
-				PurposeInfo->bIsSorted = false;
-				return true;
+				if (!Columns.IsEmpty())
+				{
+					PurposeInfo->Factories.Emplace(MoveTemp(Constructor), MoveTemp(Columns));
+					PurposeInfo->bIsSorted = false;
+					return true;
+				}
+				else
+				{
+					return false;
+				}
 			default:
 				checkf(false, TEXT("Unexpected ITypedElementDataStorageUiInterface::EPurposeType found provided when registering widget factory."));
 				return false;
@@ -275,25 +298,9 @@ void UTypedElementDatabaseUi::CreateWidgetConstructors(FName Purpose, EMatchAppr
 			PurposeInfo->Factories.StableSort(
 				[](const FWidgetFactory& Lhs, const FWidgetFactory& Rhs)
 				{
-					int32 LeftSize = Lhs.Columns.Num();
-					int32 RightSize = Rhs.Columns.Num();
-					if (LeftSize == RightSize)
-					{
-						for (int32 Index = 0; Index < LeftSize; ++Index)
-						{
-							const UScriptStruct* LhsStruct = Lhs.Columns[Index].Get();
-							const UScriptStruct* RhsStruct = Rhs.Columns[Index].Get();
-							if (LhsStruct != RhsStruct)
-							{
-								return LhsStruct < RhsStruct;
-							}
-						}
-						return true;
-					}
-					else
-					{
-						return LeftSize > RightSize;
-					}
+					int32 LeftSize = Lhs.Columns.MinimumColumnMatchRequired();
+					int32 RightSize = Rhs.Columns.MinimumColumnMatchRequired();
+					return LeftSize > RightSize;
 				});
 			PurposeInfo->bIsSorted = true;
 		}
@@ -486,83 +493,48 @@ void UTypedElementDatabaseUi::CreateStandardArchetypes()
 		}), FName(TEXT("Editor_WidgetTable")));
 }
 
-bool UTypedElementDatabaseUi::PrepareColumnsList(TArray<TWeakObjectPtr<const UScriptStruct>>& Columns)
-{
-	Columns.Sort(
-		[](const TWeakObjectPtr<const UScriptStruct>& Lhs, const TWeakObjectPtr<const UScriptStruct>& Rhs)
-		{
-			return Lhs.Get() < Rhs.Get();
-		});
-
-	// Check if there are any columns to register and whether there are any invalid columns. After sorting
-	// invalid columns will always be at the start of the array.
-	if (Columns.IsEmpty() || Columns[0] == nullptr)
-	{
-		return false;
-	}
-
-	// Remove duplicates.
-	for (int32 Index = 1; Index < Columns.Num(); )
-	{
-		if (Columns[Index - 1] != Columns[Index])
-		{
-			++Index;
-		}
-		else
-		{
-			Columns.RemoveAt(Index);
-		}
-	}
-	return true;
-}
-
 void UTypedElementDatabaseUi::CreateWidgetConstructors_LongestMatch(const TArray<FWidgetFactory>& WidgetFactories,
 	TArray<TWeakObjectPtr<const UScriptStruct>>& Columns, TypedElementDataStorage::FMetaDataView Arguments,
 	const WidgetConstructorCallback& Callback)
 {
-	using ColumnView = TConstArrayView<TWeakObjectPtr<const UScriptStruct>>;
-
+	TArray<TWeakObjectPtr<const UScriptStruct>> MatchedColumns;
 	for (auto FactoryIt = WidgetFactories.CreateConstIterator(); FactoryIt && !Columns.IsEmpty(); ++FactoryIt)
 	{
 		int32 MatchingIndex = INDEX_NONE;
-		if (FactoryIt->Columns.Num() > Columns.Num())
+		if (FactoryIt->Columns.MinimumColumnMatchRequired() > Columns.Num())
 		{
 			// There are more columns required for this factory than there are in the requested columns list so skip this
 			// factory.
 			continue;
 		}
 
-		ColumnView ColumnsRange = Columns;
-		for (auto ColumnIt = FactoryIt->Columns.CreateConstIterator(); ColumnIt && !ColumnsRange.IsEmpty(); ++ColumnIt)
-		{
-			int32 Index = INDEX_NONE;
-			if (ColumnsRange.Find(*ColumnIt, Index))
-			{
-				MatchingIndex = MatchingIndex == INDEX_NONE ? Index : MatchingIndex;
-				ColumnsRange = ColumnsRange.RightChop(Index + 1);
-			}
-			else
-			{
-				// Reset as there is at least one column that didn't match.
-				MatchingIndex = INDEX_NONE;
-				break;
-			}
-		}
+		MatchedColumns.Reset();
 
-		// This assumes that the column arrays are sorted from longest to shortest.
-		if (MatchingIndex != INDEX_NONE)
+		if (FactoryIt->Columns.Verify(MatchedColumns, Columns, true))
 		{
-			// Delete in reverse order to reduce the number of objects that need to be moved up as the order needs to remain stable.
-			ColumnsRange = ColumnView(Columns);
-			for (auto It = FactoryIt->Columns.rbegin(); It != FactoryIt->Columns.rend(); ++It)
+			// Remove the found columns from the requested list.
+			Algo::SortBy(MatchedColumns, [](const TWeakObjectPtr<const UScriptStruct>& Column) { return Column.Get(); });
+			TConstArrayView<TWeakObjectPtr<const UScriptStruct>> CleanMatchedColumns(MatchedColumns.GetData(), Algo::Unique(MatchedColumns));
+
+			TWeakObjectPtr<const UScriptStruct>* ColumnsIt = Columns.GetData();
+			TWeakObjectPtr<const UScriptStruct>* ColumnsEnd = ColumnsIt + Columns.Num();
+			int32 ColumnIndex = 0;
+			for (const TWeakObjectPtr<const UScriptStruct>& MatchedColumn : CleanMatchedColumns)
 			{
-				int32 Index = INDEX_NONE;
-				ensureMsgf(ColumnsRange.Find(*It, Index), TEXT("A previous found matching column can't be found in the original array."));
-				Columns.RemoveAt(Index);
-				ColumnsRange = ColumnsRange.Left(Index);
+				while (*ColumnsIt != MatchedColumn)
+				{
+					++ColumnIndex;
+					++ColumnsIt;
+					if (ColumnsIt == ColumnsEnd)
+					{
+						ensureMsgf(false, TEXT("A previous found matching column can't be found in the original array."));
+						return;
+					}
+				}
+				Columns.RemoveAt(ColumnIndex);
 			}
-
-			if (!CreateSingleWidgetConstructor(FactoryIt->Constructor, Arguments, FactoryIt->Columns, Callback))
+			
+			if (!CreateSingleWidgetConstructor(FactoryIt->Constructor, Arguments, CleanMatchedColumns, Callback))
 			{
 				return;
 			}
@@ -575,31 +547,46 @@ void UTypedElementDatabaseUi::CreateWidgetConstructors_ExactMatch(const TArray<F
 	const WidgetConstructorCallback& Callback)
 {
 	int32 ColumnCount = Columns.Num();
+	TArray<TWeakObjectPtr<const UScriptStruct>> MatchedColumns;
 	for (const FWidgetFactory& Factory : WidgetFactories)
 	{
-		if (Factory.Columns.Num() != ColumnCount)
+		// If there are more matches required that there are columns, then there will never be an exact match.
+		// Less than the column count can still result in a match that covers all columns.
+		if (Factory.Columns.MinimumColumnMatchRequired() > ColumnCount)
 		{
 			continue;
 		}
 
-		bool bFoundMatch = true;
-		TWeakObjectPtr<const UScriptStruct>* ColumnIt = Columns.GetData();
-		const TWeakObjectPtr<const UScriptStruct>* FactoryColumnsIt = Factory.Columns.GetData();
-		for (int32 Index = 0; Index < ColumnCount; ++Index)
+		MatchedColumns.Reset();
+
+		if (Factory.Columns.Verify(MatchedColumns, Columns, true))
 		{
-			if (ColumnIt != FactoryColumnsIt)
+			Algo::SortBy(MatchedColumns, [](const TWeakObjectPtr<const UScriptStruct>& Column) { return Column.Get(); });
+			TWeakObjectPtr<const UScriptStruct>* MatchedColumnsIt = MatchedColumns.GetData();
+			TWeakObjectPtr<const UScriptStruct>* MatchedColumnsEnd = MatchedColumnsIt + MatchedColumns.Num();
+			bool bFullyMatched = true;
+			for (TWeakObjectPtr<const UScriptStruct> Column : Columns)
 			{
-				bFoundMatch = false;
-				break;
+				while (*MatchedColumnsIt != Column)
+				{
+					if (++MatchedColumnsIt == MatchedColumnsEnd)
+					{
+						bFullyMatched = false;
+						break;
+					}
+				}
+				if (!bFullyMatched)
+				{
+					break;
+				}
+			}
+
+			if (bFullyMatched)
+			{
+				Columns.Reset();
+				CreateSingleWidgetConstructor(Factory.Constructor, Arguments, MatchedColumns, Callback);
+				return;
 			}
-			++ColumnIt;
-			++FactoryColumnsIt;
-		}
-		if (bFoundMatch)
-		{
-			Columns.Reset();
-			CreateSingleWidgetConstructor(Factory.Constructor, Arguments, Factory.Columns, Callback);
-			return;
 		}
 	}
 }
@@ -620,31 +607,25 @@ void UTypedElementDatabaseUi::CreateWidgetConstructors_SingleMatch(const TArray<
 		const UScriptStruct* ColumnType = (*ColumnIt).Get();
 		for (; FactoryIt != FactoryEnd; ++FactoryIt)
 		{
-			int32 ColumnCount = (*FactoryIt).Columns.Num();
-			if (ColumnCount > 1)
+			TConstArrayView<TWeakObjectPtr<const UScriptStruct>> ColumnData = (*FactoryIt).Columns.GetColumns();
+			if (ColumnData.Num() > 1)
 			{
 				// Moved passed the point where factories only have a single column.
 				return;
 			}
-			else if (ColumnCount == 0)
+			else if (ColumnData.Num() == 0)
 			{
 				// Need to move further to find factories with exactly one column.
 				continue;
 			}
 
-			const UScriptStruct* FactoryType = (*FactoryIt).Columns[0].Get();
-			if (FactoryType == ColumnType)
+			if (ColumnData[0] == *ColumnIt)
 			{
 				Columns.RemoveAt(ColumnIndex);
-				CreateSingleWidgetConstructor((*FactoryIt).Constructor, Arguments, (*FactoryIt).Columns, Callback);
+				CreateSingleWidgetConstructor((*FactoryIt).Constructor, Arguments, ColumnData, Callback);
 				// Match was found so move on to the next column in the column.
 				break;
 			}
-			if (FactoryType < ColumnType)
-			{
-				// The current column isn't found in the factories, so move on to the next column.
-				break;
-			}
 		}
 	}
 }
@@ -668,14 +649,14 @@ UTypedElementDatabaseUi::FWidgetFactory::FWidgetFactory(TUniquePtr<FTypedElement
 }
 
 UTypedElementDatabaseUi::FWidgetFactory::FWidgetFactory(const UScriptStruct* InConstructor, 
-	TArray<TWeakObjectPtr<const UScriptStruct>>&& InColumns)
+	TypedElementQueryBuilder::FQueryConditions&& InColumns)
 	: Columns(MoveTemp(InColumns))
 	, Constructor(InConstructor)
 {
 }
 
 UTypedElementDatabaseUi::FWidgetFactory::FWidgetFactory(TUniquePtr<FTypedElementWidgetConstructor>&& InConstructor, 
-	TArray<TWeakObjectPtr<const UScriptStruct>>&& InColumns)
+	TypedElementQueryBuilder::FQueryConditions&& InColumns)
 	: Columns(MoveTemp(InColumns))
 	, Constructor(MoveTemp(InConstructor))
 {
diff --git a/Engine/Plugins/Experimental/TypedElementsDataStorage/Source/TypedElementsDataStorage/Private/TypedElementDatabaseUI.h b/Engine/Plugins/Experimental/TypedElementsDataStorage/Source/TypedElementsDataStorage/Private/TypedElementDatabaseUI.h
--- a/Engine/Plugins/Experimental/TypedElementsDataStorage/Source/TypedElementsDataStorage/Private/TypedElementDatabaseUI.h
+++ b/Engine/Plugins/Experimental/TypedElementsDataStorage/Source/TypedElementsDataStorage/Private/TypedElementDatabaseUI.h
@@ -34,10 +34,10 @@ public:
 
 	bool RegisterWidgetFactory(FName Purpose, const UScriptStruct* Constructor) override;
 	bool RegisterWidgetFactory(FName Purpose, const UScriptStruct* Constructor,
-		TArray<TWeakObjectPtr<const UScriptStruct>> Columns) override;
+		TypedElementQueryBuilder::FQueryConditions Columns) override;
 	bool RegisterWidgetFactory(FName Purpose, TUniquePtr<FTypedElementWidgetConstructor>&& Constructor) override;
 	bool RegisterWidgetFactory(FName Purpose, TUniquePtr<FTypedElementWidgetConstructor>&& Constructor,
-		TArray<TWeakObjectPtr<const UScriptStruct>> Columns) override;
+		TypedElementQueryBuilder::FQueryConditions Columns) override;
 
 	void CreateWidgetConstructors(FName Purpose,
 		TypedElementDataStorage::FMetaDataView Arguments, const WidgetConstructorCallback& Callback) override;
@@ -56,14 +56,14 @@ private:
 	{
 		using ConstructorType = std::variant<const UScriptStruct*, TUniquePtr<FTypedElementWidgetConstructor>>;
 
-		TArray<TWeakObjectPtr<const UScriptStruct>> Columns;
+		TypedElementQueryBuilder::FQueryConditions Columns;
 		ConstructorType Constructor;
 
 		FWidgetFactory() = default;
 		explicit FWidgetFactory(const UScriptStruct* InConstructor);
 		explicit FWidgetFactory(TUniquePtr<FTypedElementWidgetConstructor>&& InConstructor);
-		FWidgetFactory(const UScriptStruct* InConstructor, TArray<TWeakObjectPtr<const UScriptStruct>>&& InColumns);
-		FWidgetFactory(TUniquePtr<FTypedElementWidgetConstructor>&& InConstructor, TArray<TWeakObjectPtr<const UScriptStruct>>&& InColumns);
+		FWidgetFactory(const UScriptStruct* InConstructor, TypedElementQueryBuilder::FQueryConditions&& InColumns);
+		FWidgetFactory(TUniquePtr<FTypedElementWidgetConstructor>&& InConstructor, TypedElementQueryBuilder::FQueryConditions&& InColumns);
 	};
 
 	struct FPurposeInfo
@@ -97,8 +97,6 @@ private:
 		TypedElementDataStorage::FMetaDataView Arguments,
 		const WidgetCreatedCallback& ConstructionCallback);
 
-	static bool PrepareColumnsList(TArray<TWeakObjectPtr<const UScriptStruct>>& Columns);
-
 	void CreateWidgetConstructors_LongestMatch(
 		const TArray<FWidgetFactory>& WidgetFactories, 
 		TArray<TWeakObjectPtr<const UScriptStruct>>& Columns, 
diff --git a/Engine/Plugins/Experimental/TypedElementsDataStorage/Source/TypedElementsDataStorageUI/Private/TypedElementLabelWidget.cpp b/Engine/Plugins/Experimental/TypedElementsDataStorage/Source/TypedElementsDataStorageUI/Private/TypedElementLabelWidget.cpp
--- a/Engine/Plugins/Experimental/TypedElementsDataStorage/Source/TypedElementsDataStorageUI/Private/TypedElementLabelWidget.cpp
+++ b/Engine/Plugins/Experimental/TypedElementsDataStorage/Source/TypedElementsDataStorageUI/Private/TypedElementLabelWidget.cpp
@@ -119,10 +119,12 @@ void UTypedElementLabelWidgetFactory::RegisterQueries(ITypedElementDataStorageIn
 void UTypedElementLabelWidgetFactory::RegisterWidgetConstructors(ITypedElementDataStorageInterface& DataStorage,
 	ITypedElementDataStorageUiInterface& DataStorageUi) const
 {
-	 DataStorageUi.RegisterWidgetFactory(FName(TEXT("General.Cell")), FTypedElementLabelWidgetConstructor::StaticStruct(),
-		{ FTypedElementLabelColumn::StaticStruct() });
-	DataStorageUi.RegisterWidgetFactory(FName(TEXT("General.Cell")), FTypedElementLabelWithHashTooltipWidgetConstructor::StaticStruct(),
-		{ FTypedElementLabelColumn::StaticStruct(), FTypedElementLabelHashColumn::StaticStruct() });
+	using namespace TypedElementQueryBuilder;
+
+	DataStorageUi.RegisterWidgetFactory<FTypedElementLabelWidgetConstructor>(FName(TEXT("General.Cell")), 
+		FColumn<FTypedElementLabelColumn>());
+	DataStorageUi.RegisterWidgetFactory<FTypedElementLabelWithHashTooltipWidgetConstructor>(FName(TEXT("General.Cell")),
+		FColumn<FTypedElementLabelColumn>() && FColumn<FTypedElementLabelHashColumn>());
 }
 
 
diff --git a/Engine/Plugins/Experimental/TypedElementsDataStorage/Source/TypedElementsDataStorageUI/Private/TypedElementPackagePathWidget.cpp b/Engine/Plugins/Experimental/TypedElementsDataStorage/Source/TypedElementsDataStorageUI/Private/TypedElementPackagePathWidget.cpp
--- a/Engine/Plugins/Experimental/TypedElementsDataStorage/Source/TypedElementsDataStorageUI/Private/TypedElementPackagePathWidget.cpp
+++ b/Engine/Plugins/Experimental/TypedElementsDataStorage/Source/TypedElementsDataStorageUI/Private/TypedElementPackagePathWidget.cpp
@@ -17,10 +17,12 @@
 void UTypedElementPackagePathWidgetFactory::RegisterWidgetConstructors(ITypedElementDataStorageInterface& DataStorage,
 	ITypedElementDataStorageUiInterface& DataStorageUi) const
 {
-	DataStorageUi.RegisterWidgetFactory(FName(TEXT("General.Cell")), FTypedElementPackagePathWidgetConstructor::StaticStruct(),
-		{ FTypedElementPackagePathColumn::StaticStruct() });
-	DataStorageUi.RegisterWidgetFactory(FName(TEXT("General.Cell")), FTypedElementLoadedPackagePathWidgetConstructor::StaticStruct(),
-		{ FTypedElementPackageLoadedPathColumn::StaticStruct() });
+	using namespace TypedElementQueryBuilder;
+
+	DataStorageUi.RegisterWidgetFactory<FTypedElementPackagePathWidgetConstructor>(FName(TEXT("General.Cell")), 
+		FColumn<FTypedElementPackagePathColumn>());
+	DataStorageUi.RegisterWidgetFactory<FTypedElementLoadedPackagePathWidgetConstructor>(FName(TEXT("General.Cell")),
+		FColumn<FTypedElementPackageLoadedPathColumn>());
 }
 
 
diff --git a/Engine/Plugins/Experimental/TypedElementsDataStorage/Source/TypedElementsDataStorageUI/Private/TypedElementTransformHeadsUpWidget.cpp b/Engine/Plugins/Experimental/TypedElementsDataStorage/Source/TypedElementsDataStorageUI/Private/TypedElementTransformHeadsUpWidget.cpp
--- a/Engine/Plugins/Experimental/TypedElementsDataStorage/Source/TypedElementsDataStorageUI/Private/TypedElementTransformHeadsUpWidget.cpp
+++ b/Engine/Plugins/Experimental/TypedElementsDataStorage/Source/TypedElementsDataStorageUI/Private/TypedElementTransformHeadsUpWidget.cpp
@@ -191,8 +191,10 @@ void UTypedElementTransformHeadsUpWidgetFactory::RegisterQueries(ITypedElementDa
 void UTypedElementTransformHeadsUpWidgetFactory::RegisterWidgetConstructors(ITypedElementDataStorageInterface& DataStorage,
 	ITypedElementDataStorageUiInterface& DataStorageUi) const
 {
-	DataStorageUi.RegisterWidgetFactory(FName(TEXT("SceneOutliner.Cell")), FTypedElementTransformHeadsUpWidgetConstructor::StaticStruct(),
-		{ FTypedElementLocalTransformColumn::StaticStruct() });
+	using namespace TypedElementQueryBuilder;
+
+	DataStorageUi.RegisterWidgetFactory<FTypedElementTransformHeadsUpWidgetConstructor>(FName(TEXT("SceneOutliner.Cell")),
+		FColumn<FTypedElementLocalTransformColumn>());
 }
 
 
diff --git a/Engine/Source/Runtime/TypedElementFramework/Public/Elements/Interfaces/TypedElementDataStorageUiInterface.h b/Engine/Source/Runtime/TypedElementFramework/Public/Elements/Interfaces/TypedElementDataStorageUiInterface.h
--- a/Engine/Source/Runtime/TypedElementFramework/Public/Elements/Interfaces/TypedElementDataStorageUiInterface.h
+++ b/Engine/Source/Runtime/TypedElementFramework/Public/Elements/Interfaces/TypedElementDataStorageUiInterface.h
@@ -5,6 +5,7 @@
 #include "Containers/ContainersFwd.h"
 #include "Elements/Interfaces/TypedElementDataStorageInterface.h"
 #include "Elements/Framework/TypedElementMetaData.h"
+#include "Elements/Framework/TypedElementQueryBuilder.h"
 #include "Templates/UnrealTypeTraits.h"
 #include "UObject/Interface.h"
 
@@ -141,6 +142,15 @@ public:
 	 * If registration is successful true will be returned otherwise false.
 	 */
 	virtual bool RegisterWidgetFactory(FName Purpose, const UScriptStruct* Constructor) = 0;
+	/**
+	 * Registers a widget factory that will be called when the purpose it's registered under is requested.
+	 * This version registers a generic type. Construction using these are typically cheaper as they can avoid
+	 * copying the Constructor and take up less memory. The downside is that they can't store additional configuration
+	 * options. If the purpose has not been registered the factory will not be recorded and a warning will be printed.
+	 * If registration is successful true will be returned otherwise false.
+	 */
+	template<typename ConstructorType>
+	bool RegisterWidgetFactory(FName Purpose);
 	/**
 	 * Registers a widget factory that will be called when the purpose it's registered under is requested.
 	 * This version registers a generic type. Construction using these are typically cheaper as they can avoid
@@ -150,7 +160,17 @@ public:
 	 * If registration is successful true will be returned otherwise false.
 	 */
 	virtual bool RegisterWidgetFactory(FName Purpose, const UScriptStruct* Constructor,
-		TArray<TWeakObjectPtr<const UScriptStruct>> Columns) = 0;
+		TypedElementQueryBuilder::FQueryConditions Columns) = 0;
+	/**
+	 * Registers a widget factory that will be called when the purpose it's registered under is requested.
+	 * This version registers a generic type. Construction using these are typically cheaper as they can avoid
+	 * copying the Constructor and take up less memory. The downside is that they can't store additional configuration
+	 * options. If the purpose has not been registered the factory will not be recorded and a warning will be printed.
+	 * The provided columns will be used when matching the factory during widget construction.
+	 * If registration is successful true will be returned otherwise false.
+	 */
+	template<typename ConstructorType>
+	bool RegisterWidgetFactory(FName Purpose, TypedElementQueryBuilder::FQueryConditions Columns);
 	/**
 	 * Registers a widget factory that will be called when the purpose it's registered under is requested.
 	 * This version uses a previously created instance of the Constructor. The benefit of this is that it store
@@ -168,7 +188,7 @@ public:
 	 * If registration is successful true will be returned otherwise false.
 	 */
 	virtual bool RegisterWidgetFactory(FName Purpose, TUniquePtr<FTypedElementWidgetConstructor>&& Constructor, 
-		TArray<TWeakObjectPtr<const UScriptStruct>> Columns) = 0;
+		TypedElementQueryBuilder::FQueryConditions Columns) = 0;
 	
 	/** Creates widget constructors for the requested purpose. */
 	virtual void CreateWidgetConstructors(FName Purpose, 
@@ -195,3 +215,20 @@ public:
 	/** Calls the provided callback for all known registered widget purposes. */
 	virtual void ListWidgetPurposes(const WidgetPurposeCallback& Callback) const = 0;
 };
+
+
+//
+// Implementations
+//
+
+template<typename ConstructorType>
+bool ITypedElementDataStorageUiInterface::RegisterWidgetFactory(FName Purpose)
+{
+	return this->RegisterWidgetFactory(Purpose, ConstructorType::StaticStruct());
+}
+
+template<typename ConstructorType>
+bool ITypedElementDataStorageUiInterface::RegisterWidgetFactory(FName Purpose, TypedElementQueryBuilder::FQueryConditions Columns)
+{
+	return this->RegisterWidgetFactory(Purpose, ConstructorType::StaticStruct(), Columns);
+}
