Rework object reinstancing
* Previous implementation had some issues depending on the order in which the Reinstance callback was made.  If implementations of the callback removed objects from the TypedElementDatabase before the callback was serviced by the ReinstancingManager then mementoization would not occur.
* New implementation keeps own mapping between old objects and the mementos that hold the data.  This is then used by the reinstancing callback handler in the ReinstancingManager to schedule row reinstancing.
* A hook into the Compatibility system was required to allow reinstancing to be aware of object deletions and keep the old object to memento mapping up to date.  These hooks are only internal to TypedElementDataStorage and not available to the outside world.  Preferred this over adding it into Compatibility directly as there becomes a circular dependency on them.  Makes it easier to shut off reinstancing too.

diff --git a/Engine/Plugins/Experimental/TypedElementsDataStorage/Source/TypedElementsDataStorage/Private/Compatibility/TypedElementObjectReinstancingManager.cpp b/Engine/Plugins/Experimental/TypedElementsDataStorage/Source/TypedElementsDataStorage/Private/Compatibility/TypedElementObjectReinstancingManager.cpp
--- a/Engine/Plugins/Experimental/TypedElementsDataStorage/Source/TypedElementsDataStorage/Private/Compatibility/TypedElementObjectReinstancingManager.cpp
+++ b/Engine/Plugins/Experimental/TypedElementsDataStorage/Source/TypedElementsDataStorage/Private/Compatibility/TypedElementObjectReinstancingManager.cpp
@@ -7,11 +7,12 @@
 #include "Memento/TypedElementMementoRowTypes.h"
 #include "TypedElementDatabase.h"
 #include "TypedElementDatabaseCompatibility.h"
-#include "Elements/Columns/TypedElementCompatibilityColumns.h"
+#include "Memento/TypedElementMementoInterface.h"
+
+DECLARE_LOG_CATEGORY_CLASS(LogTedsObjectReinstancing, Log, Log)
 
 UTypedElementObjectReinstancingManager::UTypedElementObjectReinstancingManager()
-	: UObjectAddedObserverHandle(TypedElementInvalidQueryHandle)
-	, MementoRowBaseTable(TypedElementInvalidTableHandle)
+	: MementoRowBaseTable(TypedElementInvalidTableHandle)
 {
 }
 
@@ -22,6 +23,10 @@ void UTypedElementObjectReinstancingManager::Initialize(UTypedElementDatabase& I
 	MementoSystem = &InMementoSystem;
 
 	ReinstancingCallbackHandle = FCoreUObjectDelegates::OnObjectsReinstanced.AddUObject(this, &UTypedElementObjectReinstancingManager::HandleOnObjectsReinstanced);
+	ObjectRemovedCallbackHandle = DataStorageCompatibility->RegisterObjectRemovedCallback([this](const void* Object, const FTypedElementDatabaseCompatibilityObjectTypeInfo& TypeInfo, TypedElementRowHandle Row)
+	{
+		HandleOnObjectPreRemoved(Object, TypeInfo, Row);
+	});
 	
 	RegisterQueries();
 }
@@ -31,6 +36,7 @@ void UTypedElementObjectReinstancingManager::Deinitialize()
 	UnregisterQueries();
 	
 	FCoreUObjectDelegates::OnObjectsReinstanced.Remove(ReinstancingCallbackHandle);
+	DataStorageCompatibility->UnregisterObjectRemovedCallback(ObjectRemovedCallbackHandle);
 
 	MementoSystem = nullptr;
 	DataStorageCompatibility = nullptr;
@@ -40,61 +46,68 @@ void UTypedElementObjectReinstancingManager::Deinitialize()
 void UTypedElementObjectReinstancingManager::RegisterQueries()
 {
 	using namespace TypedElementQueryBuilder;
-	
-	UObjectAddedObserverHandle = Database->RegisterQuery(
+	using DSI = ITypedElementDataStorageInterface;
+
+	{
+		TypedElementQueryHandle QueryHandle = Database->RegisterQuery(
 		Select(
-			TEXT("Object Reinstancing: Added UObjectColumn"),
-			FObserver::OnAdd<FTypedElementUObjectColumn>(),
-			[this](TypedElementDataStorage::IQueryContext& Context, TypedElementRowHandle Row, const FTypedElementUObjectColumn& UObjectColumn)
+		TEXT("Memento cleanup"),
+		FProcessor(DSI::EQueryTickPhase::FrameEnd, Database->GetQueryTickGroupName(DSI::EQueryTickGroups::Default)),
+			[this](TypedElementDataStorage::IQueryContext& Context, const TypedElementRowHandle* Row, const FTypedElementsReinstanceableSourceObject*)
 			{
-				TypedElementRowHandle Memento;
-				if (NewInstanceToMementoMap.RemoveAndCopyValue(UObjectColumn.Object.Get(), Memento))
+				Context.RemoveRows(Context.GetRowHandles());
+				
+				TConstArrayView<const FTypedElementsReinstanceableSourceObject> ReinstanceableArray = MakeArrayView(Context.GetColumn<FTypedElementsReinstanceableSourceObject>(), Context.GetRowCount());
+				for (const FTypedElementsReinstanceableSourceObject& Reinstanceable : ReinstanceableArray)
 				{
-					Context.AddColumn<FTypedElementMementoReinstanceTarget>(Memento, FTypedElementMementoReinstanceTarget{.Target = Row});
+					OldObjectToMementoMap.Remove(Reinstanceable.Object);
 				}
+				
 			})
-			.Compile());
-
-	ExternalObjectAddedObserverHandle = Database->RegisterQuery(
-	Select(
-		TEXT("Object Reinstancing: Added ExternalObjectColumn"),
-		FObserver::OnAdd<FTypedElementExternalObjectColumn>(),
-		[this](TypedElementDataStorage::IQueryContext& Context, TypedElementRowHandle Row, const FTypedElementExternalObjectColumn& ExternalObjectColumn)
-		{
-			TypedElementRowHandle Memento;
-			if (NewInstanceToMementoMap.RemoveAndCopyValue(ExternalObjectColumn.Object, Memento))
-			{
-				Context.AddColumn<FTypedElementMementoReinstanceTarget>(Memento, FTypedElementMementoReinstanceTarget{.Target = Row});
-			}
-		})
-		.Compile());
+			.Where().All<FTypedElementMementoTag, FTypedElementMementoPopulated>()
+			.Compile()
+		);
+		check(QueryHandle != TypedElementInvalidQueryHandle);
+	}
 }
 
 void UTypedElementObjectReinstancingManager::UnregisterQueries()
 {
-	// TODO: Observer queries cannot be unregistered in Mass
-	// Database->UnregisterQuery(ExternalObjectAddedObserverHandle);
-	// Database->UnregisterQuery(UObjectAddedObserverHandle);
+	// TODO: Observer queries cannot be unregistered in Mass at this time
+}
+
+void UTypedElementObjectReinstancingManager::HandleOnObjectPreRemoved(const void* Object, const FTypedElementDatabaseCompatibilityObjectTypeInfo& TypeInfo, TypedElementRowHandle ObjectRow)
+{
+	// This is the chance to record the old object to memento
+	TypedElementRowHandle Memento = MementoSystem->CreateMemento(Database.Get());
+	Database->AddOrGetColumn(ObjectRow, FTypedElementMementoOnDelete{ .Memento = Memento });
+	
+	Database->AddOrGetColumn(Memento, FTypedElementsReinstanceableSourceObject{.Object = Object});
+
+	OldObjectToMementoMap.Add(Object, Memento);
 }
 
 void UTypedElementObjectReinstancingManager::HandleOnObjectsReinstanced(
 	const FCoreUObjectDelegates::FReplacementObjectMap& ObjectReplacementMap)
 {
-	NewInstanceToMementoMap.Empty(ObjectReplacementMap.Num());
-	
 	for (FCoreUObjectDelegates::FReplacementObjectMap::TConstIterator Iter = ObjectReplacementMap.CreateConstIterator(); Iter; ++Iter)
 	{
-		const UObject* PreDeleteObject = Iter->Key;
 		const UObject* NewInstanceObject = Iter->Value;
-		
-		TypedElementRowHandle ObjectRow = DataStorageCompatibility->FindRowWithCompatibleObject(PreDeleteObject);
-		if (ObjectRow != TypedElementInvalidRowHandle)
-		{
-			TypedElementRowHandle Memento = MementoSystem->CreateMemento(Database.Get());
-			MementoSystem->EnableMementoizeOnDelete(Database.Get(), ObjectRow, Memento);
 
-			// Keep a temporary map of the new objects that will be created to the mementos that keep their state
-			NewInstanceToMementoMap.Add(NewInstanceObject, Memento);
+		TypedElementRowHandle NewObjectRow = DataStorageCompatibility->FindRowWithCompatibleObjectExplicit(NewInstanceObject);
+		if (NewObjectRow != TypedElementInvalidRowHandle)
+		{
+			const void* PreDeleteObject = Iter->Key;
+			const TypedElementRowHandle* MementoRowPtr = OldObjectToMementoMap.Find(PreDeleteObject);
+			if (MementoRowPtr != nullptr)
+			{
+				// Kick off reinstantiation of NewObjectRow from the Memento
+				TypedElementRowHandle Memento = *MementoRowPtr;
+				if (ensureMsgf(Database->HasColumns(Memento, TConstArrayView<const UScriptStruct*>({FTypedElementMementoTag::StaticStruct()})), TEXT("Cannot reinstantiate from a non memento row")))
+				{
+					Database->AddOrGetColumn(Memento, FTypedElementMementoReinstanceTarget{ .Target = NewObjectRow });
+				}
+			}
 		}
 	}
 }
diff --git a/Engine/Plugins/Experimental/TypedElementsDataStorage/Source/TypedElementsDataStorage/Private/Compatibility/TypedElementObjectReinstancingManager.h b/Engine/Plugins/Experimental/TypedElementsDataStorage/Source/TypedElementsDataStorage/Private/Compatibility/TypedElementObjectReinstancingManager.h
--- a/Engine/Plugins/Experimental/TypedElementsDataStorage/Source/TypedElementsDataStorage/Private/Compatibility/TypedElementObjectReinstancingManager.h
+++ b/Engine/Plugins/Experimental/TypedElementsDataStorage/Source/TypedElementsDataStorage/Private/Compatibility/TypedElementObjectReinstancingManager.h
@@ -27,18 +27,35 @@ public:
 private:
 	void RegisterQueries();
 	void UnregisterQueries();
+	void HandleOnObjectPreRemoved(const void* Object, const FTypedElementDatabaseCompatibilityObjectTypeInfo& TypeInfo, TypedElementRowHandle ObjectRow);
 	void HandleOnObjectsReinstanced(const FCoreUObjectDelegates::FReplacementObjectMap& ObjectReplacementMap);
 
+	UPROPERTY()
 	TObjectPtr<UTypedElementDatabase> Database = nullptr;
+	UPROPERTY()
 	TObjectPtr<UTypedElementDatabaseCompatibility> DataStorageCompatibility = nullptr;
+	UPROPERTY()
 	TObjectPtr<UTypedElementMementoSystem> MementoSystem = nullptr;
-	
-	// Maps objects that are about to be created to the memento that will be used to reinstance them
-	TMap<const void*, TypedElementRowHandle> NewInstanceToMementoMap;
 
-	TypedElementQueryHandle UObjectAddedObserverHandle;
-	TypedElementQueryHandle ExternalObjectAddedObserverHandle;
+	// Reverse lookup that holds all populated mementos for recently deleted objects
+	// Entry removed when the memento is removed
+	TMap<const void*, TypedElementRowHandle> OldObjectToMementoMap;
 	
 	TypedElementTableHandle MementoRowBaseTable;
 	FDelegateHandle ReinstancingCallbackHandle;
+	FDelegateHandle ObjectRemovedCallbackHandle;
 };
+
+USTRUCT()
+struct FTypedElementsReinstanceableSourceObject : public FTypedElementDataStorageColumn
+{
+	GENERATED_BODY()
+
+	const void* Object;
+};
+
+USTRUCT()
+struct FTypedElementsObjectReinstanceSourceColumnInitialized : public FTypedElementDataStorageTag
+{
+	GENERATED_BODY()
+};
\ No newline at end of file
diff --git a/Engine/Plugins/Experimental/TypedElementsDataStorage/Source/TypedElementsDataStorage/Private/Memento/TypedElementMementoSystem.cpp b/Engine/Plugins/Experimental/TypedElementsDataStorage/Source/TypedElementsDataStorage/Private/Memento/TypedElementMementoSystem.cpp
--- a/Engine/Plugins/Experimental/TypedElementsDataStorage/Source/TypedElementsDataStorage/Private/Memento/TypedElementMementoSystem.cpp
+++ b/Engine/Plugins/Experimental/TypedElementsDataStorage/Source/TypedElementsDataStorage/Private/Memento/TypedElementMementoSystem.cpp
@@ -23,11 +23,6 @@ TypedElementRowHandle UTypedElementMementoSystem::CreateMemento(ITypedElementDat
 	return DataStorage->AddRow(MementoRowBaseTable);
 }
 
-void UTypedElementMementoSystem::EnableMementoizeOnDelete(ITypedElementDataStorageInterface* DataStorage, TypedElementRowHandle Row, TypedElementRowHandle Memento)
-{
-	DataStorage->AddOrGetColumn<FTypedElementMementoOnDelete>(Row, FTypedElementMementoOnDelete{ .Memento = Memento });
-}
-
 void UTypedElementMementoSystem::RegisterTables(UTypedElementDatabase& DataStorage)
 {
 	// Register tables that will be used by reinstancing
@@ -89,6 +84,19 @@ void UTypedElementMementoSystem::RegisterQueries(UTypedElementDatabase& DataStor
 		check(QueryHandle != TypedElementInvalidQueryHandle);
 	}
 
+	{
+		const TypedElementQueryHandle QueryHandle = DataStorage.RegisterQuery(
+			Select(
+				TEXT("Add Populated To Memento"),
+				FObserver::OnRemove<FTypedElementMementoOnDelete>(),
+				[](TypedElementDataStorage::IQueryContext& Context, TypedElementRowHandle Memento, const FTypedElementMementoOnDelete& MementoRow)
+				{
+					Context.AddColumns(MementoRow.Memento, TConstArrayView<const UScriptStruct*>({FTypedElementMementoPopulated::StaticStruct()}));
+				})
+		.Compile());
+		check(QueryHandle != TypedElementInvalidQueryHandle);
+	}
+
 	// Setup processors to execute memento translators on reinstantiate columns onto a target
 	// The primary query runs over all mementos with a reinstance target and a column holding memento data
 	// The subquery runs on rows targeted by the reinstance column that also have a mementoizable column
diff --git a/Engine/Plugins/Experimental/TypedElementsDataStorage/Source/TypedElementsDataStorage/Private/Memento/TypedElementMementoSystem.h b/Engine/Plugins/Experimental/TypedElementsDataStorage/Source/TypedElementsDataStorage/Private/Memento/TypedElementMementoSystem.h
--- a/Engine/Plugins/Experimental/TypedElementsDataStorage/Source/TypedElementsDataStorage/Private/Memento/TypedElementMementoSystem.h
+++ b/Engine/Plugins/Experimental/TypedElementsDataStorage/Source/TypedElementsDataStorage/Private/Memento/TypedElementMementoSystem.h
@@ -23,9 +23,6 @@ public:
 	void Deinitialize();
 
 	TypedElementRowHandle CreateMemento(ITypedElementDataStorageInterface* DataStorage);
-	
-	// Enables the given Row to be mementoized into the Memento row when it is deleted
-	void EnableMementoizeOnDelete(ITypedElementDataStorageInterface* DataStorage, TypedElementRowHandle Row, TypedElementRowHandle Memento);
 
 private:
 	void RegisterQueries(UTypedElementDatabase& DataStorage) const;
diff --git a/Engine/Plugins/Experimental/TypedElementsDataStorage/Source/TypedElementsDataStorage/Private/TypedElementDatabaseCompatibility.cpp b/Engine/Plugins/Experimental/TypedElementsDataStorage/Source/TypedElementsDataStorage/Private/TypedElementDatabaseCompatibility.cpp
--- a/Engine/Plugins/Experimental/TypedElementsDataStorage/Source/TypedElementsDataStorage/Private/TypedElementDatabaseCompatibility.cpp
+++ b/Engine/Plugins/Experimental/TypedElementsDataStorage/Source/TypedElementsDataStorage/Private/TypedElementDatabaseCompatibility.cpp
@@ -13,7 +13,6 @@
 #include "MassActorEditorSubsystem.h"
 #include "MassActorSubsystem.h"
 #include "TypedElementDataStorageProfilingMacros.h"
-#include "Memento/TypedElementMementoInterface.h"
 
 void UTypedElementDatabaseCompatibility::Initialize(ITypedElementDataStorageInterface* StorageInterface)
 {
@@ -46,6 +45,36 @@ void UTypedElementDatabaseCompatibility::RegisterDealiaserCallback(ObjectToRowDe
 	ObjectToRowDialiasers.Add(MoveTemp(Dealiaser));
 }
 
+FDelegateHandle UTypedElementDatabaseCompatibility::RegisterObjectAddedCallback(ObjectAddedCallback&& OnObjectAdded)
+{
+	FDelegateHandle Handle(FDelegateHandle::GenerateNewHandle);
+	ObjectAddedCallbackList.Emplace(MoveTemp(OnObjectAdded), Handle);
+	return Handle;
+}
+
+void UTypedElementDatabaseCompatibility::UnregisterObjectAddedCallback(FDelegateHandle Handle)
+{
+	ObjectAddedCallbackList.RemoveAll([Handle](const TPair<ObjectAddedCallback, FDelegateHandle>& Element)->bool
+	{
+		return Element.Value == Handle;
+	});
+}
+
+FDelegateHandle UTypedElementDatabaseCompatibility::RegisterObjectRemovedCallback(ObjectRemovedCallback&& OnObjectAdded)
+{
+	FDelegateHandle Handle(FDelegateHandle::GenerateNewHandle);
+	PreObjectRemovedCallbackList.Emplace(MoveTemp(OnObjectAdded), Handle);
+	return Handle;
+}
+
+void UTypedElementDatabaseCompatibility::UnregisterObjectRemovedCallback(FDelegateHandle Handle)
+{
+	PreObjectRemovedCallbackList.RemoveAll([Handle](const TPair<ObjectRemovedCallback, FDelegateHandle>& Element)->bool
+	{
+		return Element.Value == Handle;
+	});
+}
+
 TypedElementRowHandle UTypedElementDatabaseCompatibility::AddCompatibleObjectExplicit(UObject* Object)
 {
 	return AddCompatibleObjectExplicit(Object, StandardUObjectTable);
@@ -133,6 +162,11 @@ void UTypedElementDatabaseCompatibility::RemoveCompatibleObjectExplicit(UObject*
 		TypedElementRowHandle Row;
 		if (ReverseObjectLookup.RemoveAndCopyValue(Object, Row))
 		{
+			const FTypedElementClassTypeInfoColumn* TypeInfoColumn = Storage->GetColumn<FTypedElementClassTypeInfoColumn>(Row);
+			if (Storage->HasRowBeenAssigned(Row) && ensureMsgf(TypeInfoColumn, TEXT("Missing type information for removed UObject at ptr 0x%p [%s]"), Object, *Object->GetName()))
+			{
+				OnPreObjectRemoved(Object, TypeInfoColumn->TypeInfo.Get(), Row);
+			}
 			Storage->RemoveRow(Row);
 		}
 	}
@@ -143,7 +177,12 @@ void UTypedElementDatabaseCompatibility::RemoveCompatibleObjectExplicit(void* Ob
 	checkf(Storage, TEXT("Removing compatible objects is not supported before Typed Element's Database compatibility manager has been initialized."));
 	TypedElementRowHandle Row;
 	if (ReverseObjectLookup.RemoveAndCopyValue(Object, Row))
-	{		
+	{
+		const FTypedElementScriptStructTypeInfoColumn* TypeInfoColumn = Storage->GetColumn<FTypedElementScriptStructTypeInfoColumn>(Row);
+		if (Storage->HasRowBeenAssigned(Row) && ensureMsgf(TypeInfoColumn, TEXT("Missing type information for removed void* object at ptr 0x%p"), Object))
+		{
+			OnPreObjectRemoved(Object, TypeInfoColumn->TypeInfo.Get(), Row);
+		}
 		Storage->RemoveRow(Row);
 	}
 }
@@ -167,6 +206,12 @@ void UTypedElementDatabaseCompatibility::RemoveCompatibleObjectExplicit(AActor*
 			if (ActorStore && !ActorStore->IsOwnedByMass()) // Only remove actors that were externally created.
 			{
 				TypedElementRowHandle Row = Entity.AsNumber();
+
+				const FTypedElementClassTypeInfoColumn* TypeInfoColumn = Storage->GetColumn<FTypedElementClassTypeInfoColumn>(Row);
+				if (Storage->HasRowBeenAssigned(Row) && ensureMsgf(TypeInfoColumn, TEXT("Missing type information for removed actor at ptr 0x%p [%s]"), Actor, *Actor->GetName()))
+				{
+					OnPreObjectRemoved(Actor, TypeInfoColumn->TypeInfo.Get(), Row);
+				}
 				
 				ActorSubsystem->RemoveHandleForActor(Actor);
 				Storage->RemoveRow(Entity.AsNumber());
@@ -426,6 +471,8 @@ void UTypedElementDatabaseCompatibility::TickPendingActorRegistration(UWorld* Ed
 					
 					// Make sure the new row is tagged for update.
 					Storage->AddColumn<FTypedElementSyncFromWorldTag>(Row);
+
+					OnObjectAdded(ActorPtr.Get(), Actor->GetClass(), Row);
 				});
 		}
 			
@@ -458,6 +505,8 @@ void UTypedElementDatabaseCompatibility::TickPendingUObjectRegistration()
 					Storage->AddOrGetColumn<FTypedElementClassTypeInfoColumn>(Row, FTypedElementClassTypeInfoColumn{ .TypeInfo = Object->GetClass() });
 					// Make sure the new row is tagged for update.
 					Storage->AddColumn<FTypedElementSyncFromWorldTag>(Row);
+					
+					OnObjectAdded(Object.Get(), Object->GetClass(), Row);
 				});
 		}
 
@@ -489,6 +538,8 @@ void UTypedElementDatabaseCompatibility::TickPendingExternalObjectRegistration()
 					Storage->AddOrGetColumn<FTypedElementScriptStructTypeInfoColumn>(Row, FTypedElementScriptStructTypeInfoColumn{ .TypeInfo = Object.TypeInfo });
 					// Make sure the new row is tagged for update.
 					Storage->AddColumn<FTypedElementSyncFromWorldTag>(Row);
+
+					OnObjectAdded(Object.Object, Object.TypeInfo.Get(), Row);
 				});
 		}
 
@@ -554,3 +605,21 @@ void UTypedElementDatabaseCompatibility::OnObjectModified(UObject* Object)
 	// batch operation during the tick step.
 	ObjectsNeedingFullSync.FindOrAdd(Object);
 }
+
+void UTypedElementDatabaseCompatibility::OnObjectAdded(const void* Object, FTypedElementDatabaseCompatibilityObjectTypeInfo TypeInfo, TypedElementRowHandle Row) const
+{
+	for (const TPair<ObjectAddedCallback, FDelegateHandle>& CallbackPair : ObjectAddedCallbackList)
+	{
+		const ObjectAddedCallback& Callback = CallbackPair.Key;
+		Callback(Object, TypeInfo, Row);
+	}
+}
+
+void UTypedElementDatabaseCompatibility::OnPreObjectRemoved(const void* Object, FTypedElementDatabaseCompatibilityObjectTypeInfo TypeInfo, TypedElementRowHandle Row) const
+{
+	for (const TPair<ObjectRemovedCallback, FDelegateHandle>& CallbackPair : PreObjectRemovedCallbackList)
+	{
+		const ObjectRemovedCallback& Callback = CallbackPair.Key;
+		Callback(Object, TypeInfo, Row);
+	}
+}
diff --git a/Engine/Plugins/Experimental/TypedElementsDataStorage/Source/TypedElementsDataStorage/Private/TypedElementDatabaseCompatibility.h b/Engine/Plugins/Experimental/TypedElementsDataStorage/Source/TypedElementsDataStorage/Private/TypedElementDatabaseCompatibility.h
--- a/Engine/Plugins/Experimental/TypedElementsDataStorage/Source/TypedElementsDataStorage/Private/TypedElementDatabaseCompatibility.h
+++ b/Engine/Plugins/Experimental/TypedElementsDataStorage/Source/TypedElementsDataStorage/Private/TypedElementDatabaseCompatibility.h
@@ -15,6 +15,9 @@ class AActor;
 class ITypedElementDataStorageInterface;
 struct FMassActorManager;
 
+enum class ETypedElementDatabaseCompatibilityObjectType : uint8;
+struct FTypedElementDatabaseCompatibilityObjectTypeInfo;
+
 UCLASS()
 class TYPEDELEMENTSDATASTORAGE_API UTypedElementDatabaseCompatibility
 	: public UObject
@@ -22,6 +25,9 @@ class TYPEDELEMENTSDATASTORAGE_API UTypedElementDatabaseCompatibility
 {
 	GENERATED_BODY()
 public:
+	using ObjectAddedCallback = TFunction<void(const void* /*Object*/, const FTypedElementDatabaseCompatibilityObjectTypeInfo&, TypedElementRowHandle /*Row*/)>;
+	using ObjectRemovedCallback = TFunction<void(const void* /*Object*/, const FTypedElementDatabaseCompatibilityObjectTypeInfo&, TypedElementRowHandle /*Row*/)>;
+	
 	~UTypedElementDatabaseCompatibility() override = default;
 
 	void Initialize(ITypedElementDataStorageInterface* StorageInterface);
@@ -29,7 +35,11 @@ public:
 
 	void RegisterRegistrationFilter(ObjectRegistrationFilter Filter) override;
 	void RegisterDealiaserCallback(ObjectToRowDealiaser Dealiaser) override;
-
+	FDelegateHandle RegisterObjectAddedCallback(ObjectAddedCallback&& OnObjectAdded);
+	void UnregisterObjectAddedCallback(FDelegateHandle Handle);
+	FDelegateHandle RegisterObjectRemovedCallback(ObjectRemovedCallback&& OnObjectRemoved);
+	void UnregisterObjectRemovedCallback(FDelegateHandle Handle);
+	
 	TypedElementRowHandle AddCompatibleObjectExplicit(UObject* Object) override;
 	TypedElementRowHandle AddCompatibleObjectExplicit(UObject* Object, TypedElementTableHandle Table) override;
 	TypedElementRowHandle AddCompatibleObjectExplicit(AActor* Actor) override;
@@ -61,6 +71,8 @@ private:
 
 	void OnPostEditChangeProperty(UObject* Object, FPropertyChangedEvent& PropertyChangedEvent);
 	void OnObjectModified(UObject* Object);
+	void OnObjectAdded(const void* Object, FTypedElementDatabaseCompatibilityObjectTypeInfo TypeInfo, TypedElementRowHandle Row) const;
+	void OnPreObjectRemoved(const void* Object, FTypedElementDatabaseCompatibilityObjectTypeInfo TypeInfo, TypedElementRowHandle Row) const;
 	
 	template<typename AddressType>
 	struct PendingRegistration
@@ -92,6 +104,8 @@ private:
 	
 	TArray<ObjectRegistrationFilter> ObjectRegistrationFilters;
 	TArray<ObjectToRowDealiaser> ObjectToRowDialiasers;
+	TArray<TPair<ObjectAddedCallback, FDelegateHandle>> ObjectAddedCallbackList;
+	TArray<TPair<ObjectRemovedCallback, FDelegateHandle>> PreObjectRemovedCallbackList;
 
 	TypedElementTableHandle StandardActorTable{ TypedElementInvalidTableHandle };
 	TypedElementTableHandle StandardActorWithTransformTable{ TypedElementInvalidTableHandle };
@@ -110,4 +124,48 @@ private:
 
 	FDelegateHandle PostEditChangePropertyDelegateHandle;
 	FDelegateHandle ObjectModifiedDelegateHandle;
-};
\ No newline at end of file
+};
+
+enum class ETypedElementDatabaseCompatibilityObjectType : uint8
+{
+	Struct,
+	Class
+};
+
+/**
+ * Objects with type info defined in either UScriptStruct or UClass can be stored into TEDS via the
+ * ITypedElementDataStorageCompatibilityInterface
+ * This is a discriminated union which aids with callbacks made when objects are added
+ */
+struct FTypedElementDatabaseCompatibilityObjectTypeInfo
+{
+	ETypedElementDatabaseCompatibilityObjectType TypeInfoType;
+
+	union
+	{
+		const UScriptStruct* ScriptStruct;
+		const UClass* Class;
+	};
+
+	FTypedElementDatabaseCompatibilityObjectTypeInfo(const UScriptStruct* InScriptStruct)
+		: TypeInfoType(ETypedElementDatabaseCompatibilityObjectType::Struct)
+		, ScriptStruct(InScriptStruct)
+	{}
+
+	FTypedElementDatabaseCompatibilityObjectTypeInfo(const UClass* InClass)
+	: TypeInfoType(ETypedElementDatabaseCompatibilityObjectType::Class)
+	, Class(InClass)
+	{}
+
+	FName GetFName() const;
+};
+
+inline FName FTypedElementDatabaseCompatibilityObjectTypeInfo::GetFName() const
+{
+	switch(TypeInfoType)
+	{
+	case ETypedElementDatabaseCompatibilityObjectType::Struct: return ScriptStruct->GetFName();
+	case ETypedElementDatabaseCompatibilityObjectType::Class: return Class->GetFName();
+	default: return FName();
+	}
+}
\ No newline at end of file
diff --git a/Engine/Plugins/Experimental/TypedElementsDataStorage/Source/TypedElementsDataStorage/Public/Memento/TypedElementMementoInterface.h b/Engine/Plugins/Experimental/TypedElementsDataStorage/Source/TypedElementsDataStorage/Public/Memento/TypedElementMementoInterface.h
--- a/Engine/Plugins/Experimental/TypedElementsDataStorage/Source/TypedElementsDataStorage/Public/Memento/TypedElementMementoInterface.h
+++ b/Engine/Plugins/Experimental/TypedElementsDataStorage/Source/TypedElementsDataStorage/Public/Memento/TypedElementMementoInterface.h
@@ -16,4 +16,13 @@ struct FTypedElementMementoOnDelete : public FTypedElementDataStorageColumn
 	// The memento row populated when the row owning this column is deleted
 	GENERATED_BODY()
 	TypedElementRowHandle Memento;
+};
+
+/**
+ * When a memento is populated, this column will be present
+ */
+USTRUCT()
+struct FTypedElementMementoPopulated : public FTypedElementDataStorageColumn
+{
+	GENERATED_BODY()
 };
\ No newline at end of file
