Expanded TEDS MetaData.

These changes allow for multiple meta data sources to be combined.
Meta data is now also used when creating a widget constructor.

diff --git a/Engine/Plugins/Experimental/TypedElementsDataStorage/Source/TypedElementsDataStorage/Private/TypedElementDatabaseUI.cpp b/Engine/Plugins/Experimental/TypedElementsDataStorage/Source/TypedElementsDataStorage/Private/TypedElementDatabaseUI.cpp
--- a/Engine/Plugins/Experimental/TypedElementsDataStorage/Source/TypedElementsDataStorage/Private/TypedElementDatabaseUI.cpp
+++ b/Engine/Plugins/Experimental/TypedElementsDataStorage/Source/TypedElementsDataStorage/Private/TypedElementDatabaseUI.cpp
@@ -262,7 +262,7 @@ bool UTypedElementDatabaseUi::RegisterWidgetFactory(FName Purpose, TUniquePtr<FT
 }
 
 void UTypedElementDatabaseUi::CreateWidgetConstructors(FName Purpose,
-	TypedElementDataStorage::FMetaDataView Arguments, const WidgetConstructorCallback& Callback)
+	const TypedElementDataStorage::FMetaDataView& Arguments, const WidgetConstructorCallback& Callback)
 {
 	if (FPurposeInfo* PurposeInfo = WidgetPurposes.Find(Purpose))
 	{
@@ -277,7 +277,7 @@ void UTypedElementDatabaseUi::CreateWidgetConstructors(FName Purpose,
 }
 
 void UTypedElementDatabaseUi::CreateWidgetConstructors(FName Purpose, EMatchApproach MatchApproach, 
-	TArray<TWeakObjectPtr<const UScriptStruct>>& Columns, TypedElementDataStorage::FMetaDataView Arguments,
+	TArray<TWeakObjectPtr<const UScriptStruct>>& Columns, const TypedElementDataStorage::FMetaDataView& Arguments,
 	const WidgetConstructorCallback& Callback)
 {
 	if (FPurposeInfo* PurposeInfo = WidgetPurposes.Find(Purpose))
@@ -323,7 +323,7 @@ void UTypedElementDatabaseUi::CreateWidgetConstructors(FName Purpose, EMatchAppr
 	}
 }
 
-void UTypedElementDatabaseUi::ConstructWidgets(FName Purpose, TypedElementDataStorage::FMetaDataView Arguments,
+void UTypedElementDatabaseUi::ConstructWidgets(FName Purpose, const TypedElementDataStorage::FMetaDataView& Arguments,
 	const WidgetCreatedCallback& ConstructionCallback)
 {
 	if (FPurposeInfo* PurposeInfo = WidgetPurposes.Find(Purpose))
@@ -347,7 +347,7 @@ void UTypedElementDatabaseUi::ConstructWidgets(FName Purpose, TypedElementDataSt
 
 bool UTypedElementDatabaseUi::CreateSingleWidgetConstructor(
 	const FWidgetFactory::ConstructorType& Constructor,
-	TypedElementDataStorage::FMetaDataView Arguments,
+	const TypedElementDataStorage::FMetaDataView& Arguments,
 	TConstArrayView<TWeakObjectPtr<const UScriptStruct>> MatchedColumnTypes, 
 	const WidgetConstructorCallback& Callback)
 {
@@ -383,7 +383,7 @@ bool UTypedElementDatabaseUi::CreateSingleWidgetConstructor(
 
 void UTypedElementDatabaseUi::CreateWidgetInstanceFromDescription(
 	const UScriptStruct* Target,
-	TypedElementDataStorage::FMetaDataView Arguments,
+	const TypedElementDataStorage::FMetaDataView& Arguments,
 	const WidgetCreatedCallback& ConstructionCallback)
 {
 	FTypedElementWidgetConstructor* Constructor = reinterpret_cast<FTypedElementWidgetConstructor*>(
@@ -402,7 +402,7 @@ void UTypedElementDatabaseUi::CreateWidgetInstanceFromDescription(
 
 void UTypedElementDatabaseUi::CreateWidgetInstanceFromInstance(
 	FTypedElementWidgetConstructor* SourceConstructor,
-	TypedElementDataStorage::FMetaDataView Arguments,
+	const TypedElementDataStorage::FMetaDataView& Arguments,
 	const WidgetCreatedCallback& ConstructionCallback)
 {
 	if (SourceConstructor->CanBeReused())
@@ -431,7 +431,7 @@ void UTypedElementDatabaseUi::CreateWidgetInstanceFromInstance(
 
 void UTypedElementDatabaseUi::CreateWidgetInstance(
 	FTypedElementWidgetConstructor& Constructor, 
-	TypedElementDataStorage::FMetaDataView Arguments,
+	const TypedElementDataStorage::FMetaDataView& Arguments,
 	const WidgetCreatedCallback& ConstructionCallback)
 {
 	TypedElementRowHandle Row = Storage->AddRow(WidgetTable);
@@ -448,7 +448,7 @@ void UTypedElementDatabaseUi::CreateWidgetInstance(
 }
 
 TSharedPtr<SWidget> UTypedElementDatabaseUi::ConstructWidget(TypedElementRowHandle Row, FTypedElementWidgetConstructor& Constructor,
-	TypedElementDataStorage::FMetaDataView Arguments)
+	const TypedElementDataStorage::FMetaDataView& Arguments)
 {
 	if (Constructor.CanBeReused())
 	{
@@ -494,7 +494,7 @@ void UTypedElementDatabaseUi::CreateStandardArchetypes()
 }
 
 void UTypedElementDatabaseUi::CreateWidgetConstructors_LongestMatch(const TArray<FWidgetFactory>& WidgetFactories,
-	TArray<TWeakObjectPtr<const UScriptStruct>>& Columns, TypedElementDataStorage::FMetaDataView Arguments,
+	TArray<TWeakObjectPtr<const UScriptStruct>>& Columns, const TypedElementDataStorage::FMetaDataView& Arguments,
 	const WidgetConstructorCallback& Callback)
 {
 	TArray<TWeakObjectPtr<const UScriptStruct>> MatchedColumns;
@@ -543,7 +543,7 @@ void UTypedElementDatabaseUi::CreateWidgetConstructors_LongestMatch(const TArray
 }
 
 void UTypedElementDatabaseUi::CreateWidgetConstructors_ExactMatch(const TArray<FWidgetFactory>& WidgetFactories,
-	TArray<TWeakObjectPtr<const UScriptStruct>>& Columns, TypedElementDataStorage::FMetaDataView Arguments,
+	TArray<TWeakObjectPtr<const UScriptStruct>>& Columns, const TypedElementDataStorage::FMetaDataView& Arguments,
 	const WidgetConstructorCallback& Callback)
 {
 	int32 ColumnCount = Columns.Num();
@@ -592,7 +592,7 @@ void UTypedElementDatabaseUi::CreateWidgetConstructors_ExactMatch(const TArray<F
 }
 
 void UTypedElementDatabaseUi::CreateWidgetConstructors_SingleMatch(const TArray<FWidgetFactory>& WidgetFactories,
-	TArray<TWeakObjectPtr<const UScriptStruct>>& Columns, TypedElementDataStorage::FMetaDataView Arguments,
+	TArray<TWeakObjectPtr<const UScriptStruct>>& Columns, const TypedElementDataStorage::FMetaDataView& Arguments,
 	const WidgetConstructorCallback& Callback)
 {
 	auto FactoryIt = WidgetFactories.rbegin();
diff --git a/Engine/Plugins/Experimental/TypedElementsDataStorage/Source/TypedElementsDataStorage/Private/TypedElementDatabaseUI.h b/Engine/Plugins/Experimental/TypedElementsDataStorage/Source/TypedElementsDataStorage/Private/TypedElementDatabaseUI.h
--- a/Engine/Plugins/Experimental/TypedElementsDataStorage/Source/TypedElementsDataStorage/Private/TypedElementDatabaseUI.h
+++ b/Engine/Plugins/Experimental/TypedElementsDataStorage/Source/TypedElementsDataStorage/Private/TypedElementDatabaseUI.h
@@ -40,14 +40,14 @@ public:
 		TypedElementQueryBuilder::FQueryConditions Columns) override;
 
 	void CreateWidgetConstructors(FName Purpose,
-		TypedElementDataStorage::FMetaDataView Arguments, const WidgetConstructorCallback& Callback) override;
+		const TypedElementDataStorage::FMetaDataView& Arguments, const WidgetConstructorCallback& Callback) override;
 	void CreateWidgetConstructors(FName Purpose, EMatchApproach MatchApproach, TArray<TWeakObjectPtr<const UScriptStruct>>& Columns,
-		TypedElementDataStorage::FMetaDataView Arguments, const WidgetConstructorCallback& Callback) override;
+		const TypedElementDataStorage::FMetaDataView& Arguments, const WidgetConstructorCallback& Callback) override;
 
-	void ConstructWidgets(FName Purpose, TypedElementDataStorage::FMetaDataView Arguments,
+	void ConstructWidgets(FName Purpose, const TypedElementDataStorage::FMetaDataView& Arguments,
 		const WidgetCreatedCallback& ConstructionCallback) override;
 	TSharedPtr<SWidget> ConstructWidget(TypedElementRowHandle Row, FTypedElementWidgetConstructor& Constructor,
-		TypedElementDataStorage::FMetaDataView Arguments) override;
+		const TypedElementDataStorage::FMetaDataView& Arguments) override;
 
 	void ListWidgetPurposes(const WidgetPurposeCallback& Callback) const override;
 
@@ -78,39 +78,39 @@ private:
 
 	bool CreateSingleWidgetConstructor(
 		const FWidgetFactory::ConstructorType& Constructor,
-		TypedElementDataStorage::FMetaDataView Arguments,
+		const TypedElementDataStorage::FMetaDataView& Arguments,
 		TConstArrayView<TWeakObjectPtr<const UScriptStruct>> MatchedColumnTypes,
 		const WidgetConstructorCallback& Callback);
 
 	void CreateWidgetInstanceFromDescription(
 		const UScriptStruct* Target,
-		TypedElementDataStorage::FMetaDataView Arguments,
+		const TypedElementDataStorage::FMetaDataView& Arguments,
 		const WidgetCreatedCallback& ConstructionCallback);
 
 	void CreateWidgetInstanceFromInstance(
 		FTypedElementWidgetConstructor* SourceConstructor,
-		TypedElementDataStorage::FMetaDataView Arguments,
+		const TypedElementDataStorage::FMetaDataView& Arguments,
 		const WidgetCreatedCallback& ConstructionCallback);
 
 	void CreateWidgetInstance(
 		FTypedElementWidgetConstructor& Constructor,
-		TypedElementDataStorage::FMetaDataView Arguments,
+		const TypedElementDataStorage::FMetaDataView& Arguments,
 		const WidgetCreatedCallback& ConstructionCallback);
 
 	void CreateWidgetConstructors_LongestMatch(
 		const TArray<FWidgetFactory>& WidgetFactories, 
 		TArray<TWeakObjectPtr<const UScriptStruct>>& Columns, 
-		TypedElementDataStorage::FMetaDataView Arguments,
+		const TypedElementDataStorage::FMetaDataView& Arguments,
 		const WidgetConstructorCallback& Callback);
 	void CreateWidgetConstructors_ExactMatch(
 		const TArray<FWidgetFactory>& WidgetFactories,
 		TArray<TWeakObjectPtr<const UScriptStruct>>& Columns,
-		TypedElementDataStorage::FMetaDataView Arguments,
+		const TypedElementDataStorage::FMetaDataView& Arguments,
 		const WidgetConstructorCallback& Callback);
 	void CreateWidgetConstructors_SingleMatch(
 		const TArray<FWidgetFactory>& WidgetFactories,
 		TArray<TWeakObjectPtr<const UScriptStruct>>& Columns,
-		TypedElementDataStorage::FMetaDataView Arguments,
+		const TypedElementDataStorage::FMetaDataView& Arguments,
 		const WidgetConstructorCallback& Callback);
 
 	TypedElementTableHandle WidgetTable{ TypedElementInvalidTableHandle };
diff --git a/Engine/Plugins/Experimental/TypedElementsDataStorage/Source/TypedElementsDataStorageUI/Private/TypedElementCounterWidget.cpp b/Engine/Plugins/Experimental/TypedElementsDataStorage/Source/TypedElementsDataStorageUI/Private/TypedElementCounterWidget.cpp
--- a/Engine/Plugins/Experimental/TypedElementsDataStorage/Source/TypedElementsDataStorageUI/Private/TypedElementCounterWidget.cpp
+++ b/Engine/Plugins/Experimental/TypedElementsDataStorage/Source/TypedElementsDataStorageUI/Private/TypedElementCounterWidget.cpp
@@ -154,7 +154,7 @@ TConstArrayView<const UScriptStruct*> FTypedElementCounterWidgetConstructor::Get
 	return Columns;
 }
 
-TSharedPtr<SWidget> FTypedElementCounterWidgetConstructor::CreateWidget(TypedElementDataStorage::FMetaDataView Arguments)
+TSharedPtr<SWidget> FTypedElementCounterWidgetConstructor::CreateWidget(const TypedElementDataStorage::FMetaDataView& Arguments)
 {
 	return SNew(STextBlock)
 		.Text(FText::Format(LabelText, 0))
diff --git a/Engine/Plugins/Experimental/TypedElementsDataStorage/Source/TypedElementsDataStorageUI/Private/TypedElementExportedTextWidget.cpp b/Engine/Plugins/Experimental/TypedElementsDataStorage/Source/TypedElementsDataStorageUI/Private/TypedElementExportedTextWidget.cpp
--- a/Engine/Plugins/Experimental/TypedElementsDataStorage/Source/TypedElementsDataStorageUI/Private/TypedElementExportedTextWidget.cpp
+++ b/Engine/Plugins/Experimental/TypedElementsDataStorage/Source/TypedElementsDataStorageUI/Private/TypedElementExportedTextWidget.cpp
@@ -110,7 +110,7 @@ TConstArrayView<const UScriptStruct*> FTypedElementExportedTextWidgetConstructor
 	return Columns;
 }
 
-TSharedPtr<SWidget> FTypedElementExportedTextWidgetConstructor::CreateWidget(TypedElementDataStorage::FMetaDataView Arguments)
+TSharedPtr<SWidget> FTypedElementExportedTextWidgetConstructor::CreateWidget(const TypedElementDataStorage::FMetaDataView& Arguments)
 {
 	return SNew(STextBlock);
 }
diff --git a/Engine/Plugins/Experimental/TypedElementsDataStorage/Source/TypedElementsDataStorageUI/Private/TypedElementLabelWidget.cpp b/Engine/Plugins/Experimental/TypedElementsDataStorage/Source/TypedElementsDataStorageUI/Private/TypedElementLabelWidget.cpp
--- a/Engine/Plugins/Experimental/TypedElementsDataStorage/Source/TypedElementsDataStorageUI/Private/TypedElementLabelWidget.cpp
+++ b/Engine/Plugins/Experimental/TypedElementsDataStorage/Source/TypedElementsDataStorageUI/Private/TypedElementLabelWidget.cpp
@@ -152,7 +152,7 @@ TConstArrayView<const UScriptStruct*> FTypedElementLabelWidgetConstructor::GetAd
 	return Columns;
 }
 
-TSharedPtr<SWidget> FTypedElementLabelWidgetConstructor::CreateWidget(TypedElementDataStorage::FMetaDataView Arguments)
+TSharedPtr<SWidget> FTypedElementLabelWidgetConstructor::CreateWidget(const TypedElementDataStorage::FMetaDataView& Arguments)
 {
 	const bool* IsEditable = Arguments.FindForColumn<FTypedElementLabelColumn>(TypedElementDataStorage::IsEditableName).TryGetExact<bool>();
 	return SNew(STextBlock).IsEnabled(IsEditable ? *IsEditable : false);
diff --git a/Engine/Plugins/Experimental/TypedElementsDataStorage/Source/TypedElementsDataStorageUI/Private/TypedElementPackagePathWidget.cpp b/Engine/Plugins/Experimental/TypedElementsDataStorage/Source/TypedElementsDataStorageUI/Private/TypedElementPackagePathWidget.cpp
--- a/Engine/Plugins/Experimental/TypedElementsDataStorage/Source/TypedElementsDataStorageUI/Private/TypedElementPackagePathWidget.cpp
+++ b/Engine/Plugins/Experimental/TypedElementsDataStorage/Source/TypedElementsDataStorageUI/Private/TypedElementPackagePathWidget.cpp
@@ -41,7 +41,7 @@ FTypedElementPackagePathWidgetConstructor::FTypedElementPackagePathWidgetConstru
 {
 }
 
-TSharedPtr<SWidget> FTypedElementPackagePathWidgetConstructor::CreateWidget(TypedElementDataStorage::FMetaDataView Arguments)
+TSharedPtr<SWidget> FTypedElementPackagePathWidgetConstructor::CreateWidget(const TypedElementDataStorage::FMetaDataView& Arguments)
 {
 	return SNew(STextBlock)
 		.OverflowPolicy(ETextOverflowPolicy::Ellipsis)
diff --git a/Engine/Plugins/Experimental/TypedElementsDataStorage/Source/TypedElementsDataStorageUI/Private/TypedElementTransformHeadsUpWidget.cpp b/Engine/Plugins/Experimental/TypedElementsDataStorage/Source/TypedElementsDataStorageUI/Private/TypedElementTransformHeadsUpWidget.cpp
--- a/Engine/Plugins/Experimental/TypedElementsDataStorage/Source/TypedElementsDataStorageUI/Private/TypedElementTransformHeadsUpWidget.cpp
+++ b/Engine/Plugins/Experimental/TypedElementsDataStorage/Source/TypedElementsDataStorageUI/Private/TypedElementTransformHeadsUpWidget.cpp
@@ -215,7 +215,7 @@ TConstArrayView<const UScriptStruct*> FTypedElementTransformHeadsUpWidgetConstru
 	return Columns;
 }
 
-TSharedPtr<SWidget> FTypedElementTransformHeadsUpWidgetConstructor::CreateWidget(TypedElementDataStorage::FMetaDataView Arguments)
+TSharedPtr<SWidget> FTypedElementTransformHeadsUpWidgetConstructor::CreateWidget(const TypedElementDataStorage::FMetaDataView& Arguments)
 {
 	return SNew(STransformQuickDisplay);
 }
diff --git a/Engine/Plugins/Experimental/TypedElementsDataStorage/Source/TypedElementsDataStorageUI/Public/TypedElementCounterWidget.h b/Engine/Plugins/Experimental/TypedElementsDataStorage/Source/TypedElementsDataStorageUI/Public/TypedElementCounterWidget.h
--- a/Engine/Plugins/Experimental/TypedElementsDataStorage/Source/TypedElementsDataStorageUI/Public/TypedElementCounterWidget.h
+++ b/Engine/Plugins/Experimental/TypedElementsDataStorage/Source/TypedElementsDataStorageUI/Public/TypedElementCounterWidget.h
@@ -54,7 +54,7 @@ public:
 	TypedElementQueryHandle Query;
 
 protected:
-	TSharedPtr<SWidget> CreateWidget(TypedElementDataStorage::FMetaDataView Arguments) override;
+	TSharedPtr<SWidget> CreateWidget(const TypedElementDataStorage::FMetaDataView& Arguments) override;
 	bool SetColumns(ITypedElementDataStorageInterface* DataStorage, TypedElementRowHandle Row) override;
 };
 
diff --git a/Engine/Plugins/Experimental/TypedElementsDataStorage/Source/TypedElementsDataStorageUI/Public/TypedElementExportedTextWidget.h b/Engine/Plugins/Experimental/TypedElementsDataStorage/Source/TypedElementsDataStorageUI/Public/TypedElementExportedTextWidget.h
--- a/Engine/Plugins/Experimental/TypedElementsDataStorage/Source/TypedElementsDataStorageUI/Public/TypedElementExportedTextWidget.h
+++ b/Engine/Plugins/Experimental/TypedElementsDataStorage/Source/TypedElementsDataStorageUI/Public/TypedElementExportedTextWidget.h
@@ -36,7 +36,7 @@ public:
 	TConstArrayView<const UScriptStruct*> GetAdditionalColumnsList() const override;
 
 protected:
-	TSharedPtr<SWidget> CreateWidget(TypedElementDataStorage::FMetaDataView Arguments) override;
+	TSharedPtr<SWidget> CreateWidget(const TypedElementDataStorage::FMetaDataView& Arguments) override;
 	bool FinalizeWidget(
 		ITypedElementDataStorageInterface* DataStorage,
 		ITypedElementDataStorageUiInterface* DataStorageUi,
diff --git a/Engine/Plugins/Experimental/TypedElementsDataStorage/Source/TypedElementsDataStorageUI/Public/TypedElementLabelWidget.h b/Engine/Plugins/Experimental/TypedElementsDataStorage/Source/TypedElementsDataStorageUI/Public/TypedElementLabelWidget.h
--- a/Engine/Plugins/Experimental/TypedElementsDataStorage/Source/TypedElementsDataStorageUI/Public/TypedElementLabelWidget.h
+++ b/Engine/Plugins/Experimental/TypedElementsDataStorage/Source/TypedElementsDataStorageUI/Public/TypedElementLabelWidget.h
@@ -36,7 +36,7 @@ public:
 
 protected:
 	explicit FTypedElementLabelWidgetConstructor(const UScriptStruct* InTypeInfo);
-	TSharedPtr<SWidget> CreateWidget(TypedElementDataStorage::FMetaDataView Arguments) override;
+	TSharedPtr<SWidget> CreateWidget(const TypedElementDataStorage::FMetaDataView& Arguments) override;
 	bool SetColumns(ITypedElementDataStorageInterface* DataStorage, TypedElementRowHandle Row) override;
 	bool FinalizeWidget(ITypedElementDataStorageInterface* DataStorage, ITypedElementDataStorageUiInterface* DataStorageUi,
 		TypedElementRowHandle Row, const TSharedPtr<SWidget>& Widget) override;
diff --git a/Engine/Plugins/Experimental/TypedElementsDataStorage/Source/TypedElementsDataStorageUI/Public/TypedElementPackagePathWidget.h b/Engine/Plugins/Experimental/TypedElementsDataStorage/Source/TypedElementsDataStorageUI/Public/TypedElementPackagePathWidget.h
--- a/Engine/Plugins/Experimental/TypedElementsDataStorage/Source/TypedElementsDataStorageUI/Public/TypedElementPackagePathWidget.h
+++ b/Engine/Plugins/Experimental/TypedElementsDataStorage/Source/TypedElementsDataStorageUI/Public/TypedElementPackagePathWidget.h
@@ -34,7 +34,7 @@ public:
 protected:
 	explicit FTypedElementPackagePathWidgetConstructor(const UScriptStruct* InTypeInfo);
 
-	TSharedPtr<SWidget> CreateWidget(TypedElementDataStorage::FMetaDataView Arguments) override;
+	TSharedPtr<SWidget> CreateWidget(const TypedElementDataStorage::FMetaDataView& Arguments) override;
 	bool FinalizeWidget(ITypedElementDataStorageInterface* DataStorage, ITypedElementDataStorageUiInterface* DataStorageUi,
 		TypedElementRowHandle Row, const TSharedPtr<SWidget>& Widget) override;
 };
diff --git a/Engine/Plugins/Experimental/TypedElementsDataStorage/Source/TypedElementsDataStorageUI/Public/TypedElementTransformHeadsUpWidget.h b/Engine/Plugins/Experimental/TypedElementsDataStorage/Source/TypedElementsDataStorageUI/Public/TypedElementTransformHeadsUpWidget.h
--- a/Engine/Plugins/Experimental/TypedElementsDataStorage/Source/TypedElementsDataStorageUI/Public/TypedElementTransformHeadsUpWidget.h
+++ b/Engine/Plugins/Experimental/TypedElementsDataStorage/Source/TypedElementsDataStorageUI/Public/TypedElementTransformHeadsUpWidget.h
@@ -41,7 +41,7 @@ public:
 	TConstArrayView<const UScriptStruct*> GetAdditionalColumnsList() const override;
 
 protected:
-	TSharedPtr<SWidget> CreateWidget(TypedElementDataStorage::FMetaDataView Arguments) override;
+	TSharedPtr<SWidget> CreateWidget(const TypedElementDataStorage::FMetaDataView& Arguments) override;
 	bool FinalizeWidget(
 		ITypedElementDataStorageInterface* DataStorage,
 		ITypedElementDataStorageUiInterface* DataStorageUi,
diff --git a/Engine/Source/Runtime/TypedElementFramework/Private/Elements/Framework/TypedElementMetaData.cpp b/Engine/Source/Runtime/TypedElementFramework/Private/Elements/Framework/TypedElementMetaData.cpp
--- a/Engine/Source/Runtime/TypedElementFramework/Private/Elements/Framework/TypedElementMetaData.cpp
+++ b/Engine/Source/Runtime/TypedElementFramework/Private/Elements/Framework/TypedElementMetaData.cpp
@@ -118,43 +118,112 @@ namespace TypedElementDataStorage
 	//
 	// FMetaDataView
 	//
+	
+	FMetaDataEntryView FMetaDataView::FindGeneric(FName AttributeName) const
+	{
+		return FMetaDataEntryView();
+	}
 
-	FMetaDataView::FMetaDataView(const TypedElementDataStorage::FQueryDescription& InQuery)
-		: Query(&InQuery)
-	{}
+	FMetaDataEntryView FMetaDataView::FindForColumn(TWeakObjectPtr<const UScriptStruct> Column, FName AttributeName) const
+	{
+		return FMetaDataEntryView();
+	}
 
-	FMetaDataView::FMetaDataView(const FMetaData& InQueryWideMetaData)
-		: QueryWideMetaData(&InQueryWideMetaData)
-	{}
 
-	FMetaDataView::FMetaDataView(const TypedElementDataStorage::FQueryDescription& InQuery, const FMetaData& InQueryWideMetaData)
-		: Query(&InQuery)
-		, QueryWideMetaData(&InQueryWideMetaData)
-	{}
 
-	FMetaDataEntryView FMetaDataView::FindGeneric(FName AttributeName) const
+	//
+	// FQueryMetaDataView
+	//
+
+	FQueryMetaDataView::FQueryMetaDataView(const TypedElementDataStorage::FQueryDescription& InQuery)
+		:Query(InQuery)
+	{
+	}
+
+	FMetaDataEntryView FQueryMetaDataView::FindGeneric(FName AttributeName) const
+	{
+		return Query.MetaData.Find(AttributeName);
+	}
+
+	FMetaDataEntryView FQueryMetaDataView::FindForColumn(
+		TWeakObjectPtr<const UScriptStruct> Column, FName AttributeName) const
+	{
+		int32 Index = 0;
+		if (Query.SelectionTypes.Find(Column, Index))
+		{
+			FMetaDataEntryView Result = Query.SelectionMetaData[Index].Find(AttributeName);
+			if (Result.IsSet())
+			{
+				return Result;
+			}
+		}
+
+		return FMetaDataEntryView();
+	}
+
+
+
+	//
+	// FColumnsMetaDataView
+	//
+
+	FColumnsMetaDataView::FColumnsMetaDataView(TConstArrayView<TWeakObjectPtr<const UScriptStruct>> InColumns)
+		: Columns(InColumns)
 	{
-		return QueryWideMetaData ? QueryWideMetaData->Find(AttributeName) : FMetaDataEntryView();
 	}
 
-	FMetaDataEntryView FMetaDataView::FindForColumn(TWeakObjectPtr<const UScriptStruct> Column, FName AttributeName, ESearchScope Scope) const
+	FMetaDataEntryView FColumnsMetaDataView::FindForColumn(
+		TWeakObjectPtr<const UScriptStruct> Column, FName AttributeName) const
 	{
-		if (Query)
+#if WITH_EDITORONLY_DATA
+		if (Columns.Find(Column) != INDEX_NONE)
 		{
-			// There are typically a small number of columns in a query, so a linear search if often fast enough and can even
-			// be faster than a map.
-			int32 Index = 0;
-			if (Query->SelectionTypes.Find(Column, Index))
+			if (const UScriptStruct* ColumnType = Column.Get())
 			{
-				FMetaDataEntryView Result = Query->SelectionMetaData[Index].Find(AttributeName);
-				if (Result.IsSet())
+				if (const FString* FoundMetaData = ColumnType->FindMetaData(AttributeName))
 				{
-					return Result;
+					return FMetaDataEntryView(*FoundMetaData);
 				}
 			}
 		}
-		return Scope == ESearchScope::FallbackOnGeneric && QueryWideMetaData
-			? QueryWideMetaData->Find(AttributeName)
-			: FMetaDataEntryView();
+#endif
+		return FMetaDataEntryView();
+	}
+
+	
+
+	//
+	// FForwardingMetaDataView
+	//
+
+	FForwardingMetaDataView::FForwardingMetaDataView(const FMetaDataView& InView)
+		: View(InView)
+	{
+	}
+
+	FMetaDataEntryView FForwardingMetaDataView::FindGeneric(FName AttributeName) const
+	{
+		return View.FindGeneric(AttributeName);
+	}
+
+	FMetaDataEntryView FForwardingMetaDataView::FindForColumn(
+		TWeakObjectPtr<const UScriptStruct> Column, FName AttributeName) const
+	{
+		return View.FindForColumn(Column, AttributeName);
+	}
+
+	
+	
+	//
+	// FGenericMetaDataView
+	//
+	FGenericMetaDataView::FGenericMetaDataView(const FMetaData& InMetaData)
+		: MetaData(InMetaData)
+	{
+	}
+	
+	FMetaDataEntryView FGenericMetaDataView::FindGeneric(FName AttributeName) const
+	{
+		return MetaData.Find(AttributeName);
 	}
 } // namespace TypedElementDataStorage
diff --git a/Engine/Source/Runtime/TypedElementFramework/Private/Elements/Framework/TypedElementQueryBuilder.cpp b/Engine/Source/Runtime/TypedElementFramework/Private/Elements/Framework/TypedElementQueryBuilder.cpp
--- a/Engine/Source/Runtime/TypedElementFramework/Private/Elements/Framework/TypedElementQueryBuilder.cpp
+++ b/Engine/Source/Runtime/TypedElementFramework/Private/Elements/Framework/TypedElementQueryBuilder.cpp
@@ -680,6 +680,7 @@ namespace TypedElementQueryBuilder
 		Query->DependencyFlags.Shrink();
 		Query->CachedDependencies.Shrink();
 		Query->Subqueries.Shrink();
+		Query->MetaData.Shrink();
 		return MoveTemp(*Query);
 	}
 
diff --git a/Engine/Source/Runtime/TypedElementFramework/Private/Elements/Interfaces/TypedElementDataStorageUiInterface.cpp b/Engine/Source/Runtime/TypedElementFramework/Private/Elements/Interfaces/TypedElementDataStorageUiInterface.cpp
--- a/Engine/Source/Runtime/TypedElementFramework/Private/Elements/Interfaces/TypedElementDataStorageUiInterface.cpp
+++ b/Engine/Source/Runtime/TypedElementFramework/Private/Elements/Interfaces/TypedElementDataStorageUiInterface.cpp
@@ -28,7 +28,7 @@ TSharedPtr<SWidget> FTypedElementWidgetConstructor::Construct(
 	TypedElementRowHandle Row,
 	ITypedElementDataStorageInterface* DataStorage,
 	ITypedElementDataStorageUiInterface* DataStorageUi,
-	TypedElementDataStorage::FMetaDataView Arguments)
+	const TypedElementDataStorage::FMetaDataView& Arguments)
 {
 	TSharedPtr<SWidget> Widget = CreateWidget(Arguments);
 	if (Widget)
diff --git a/Engine/Source/Runtime/TypedElementFramework/Public/Elements/Common/TypedElementQueryDescription.h b/Engine/Source/Runtime/TypedElementFramework/Public/Elements/Common/TypedElementQueryDescription.h
--- a/Engine/Source/Runtime/TypedElementFramework/Public/Elements/Common/TypedElementQueryDescription.h
+++ b/Engine/Source/Runtime/TypedElementFramework/Public/Elements/Common/TypedElementQueryDescription.h
@@ -98,6 +98,7 @@ namespace TypedElementDataStorage
 		/** Cached instances of the dependencies. This will always match the count of the other Dependency*Types, but may contain null pointers. */
 		TArray<TWeakObjectPtr<UObject>, TInlineAllocator<NumInlineDependencies>> CachedDependencies;
 		TArray<QueryHandle> Subqueries;
+		FMetaData MetaData;
 
 		EActionType Action;
 		/** If true, this query only has simple operations and is guaranteed to be executed fully and at optimal performance. */
diff --git a/Engine/Source/Runtime/TypedElementFramework/Public/Elements/Framework/TypedElementMetaData.h b/Engine/Source/Runtime/TypedElementFramework/Public/Elements/Framework/TypedElementMetaData.h
--- a/Engine/Source/Runtime/TypedElementFramework/Public/Elements/Framework/TypedElementMetaData.h
+++ b/Engine/Source/Runtime/TypedElementFramework/Public/Elements/Framework/TypedElementMetaData.h
@@ -8,6 +8,7 @@
 #include "Containers/Map.h"
 #include "Misc/EnumClassFlags.h"
 #include "Misc/TVariant.h"
+#include "Templates/Tuple.h"
 
 class UScriptStruct;
 
@@ -100,33 +101,105 @@ namespace TypedElementDataStorage
 	};
 
 	/**
-	 * Short lived view of a meta data container optionally associated with a query.
+	 * Short lived view of a meta data container.
 	 */
-	class FMetaDataView final
+	class FMetaDataView
 	{
 	public:
-		enum class ESearchScope
-		{
-			FallbackOnGeneric, // If the attribute isn't found on the column or the column isn't found, search the generic meta data for the attribute.
-			ColumnOnly // Only use the data on the column.
-		};
+		virtual ~FMetaDataView() = default;
 
-		FMetaDataView() = default;
-		TYPEDELEMENTFRAMEWORK_API FMetaDataView(const TypedElementDataStorage::FQueryDescription& InQuery); // Deliberately avoided "explicit".
-		TYPEDELEMENTFRAMEWORK_API FMetaDataView(const FMetaData& InQueryWideMetaData); // Deliberately avoided "explicit".
-		TYPEDELEMENTFRAMEWORK_API FMetaDataView(
-			const TypedElementDataStorage::FQueryDescription& InQuery, const FMetaData& InQueryWideMetaData);
+		TYPEDELEMENTFRAMEWORK_API virtual FMetaDataEntryView FindGeneric(FName AttributeName) const;
+		TYPEDELEMENTFRAMEWORK_API virtual FMetaDataEntryView FindForColumn(
+			TWeakObjectPtr<const UScriptStruct> Column, FName AttributeName) const;
+		template<typename Column>
+		FMetaDataEntryView FindForColumn(FName AttributeName) const;
+	};
 
-		TYPEDELEMENTFRAMEWORK_API FMetaDataEntryView FindGeneric(FName AttributeName) const;
+	/**
+	 * Short lived view of a meta data container that wraps a query.
+	 */
+	class FQueryMetaDataView final : public FMetaDataView
+	{
+	public:
+		~FQueryMetaDataView() override = default;
+
+		TYPEDELEMENTFRAMEWORK_API explicit FQueryMetaDataView(const TypedElementDataStorage::FQueryDescription& InQuery);
+		
+		TYPEDELEMENTFRAMEWORK_API FMetaDataEntryView FindGeneric(FName AttributeName) const override;
 		TYPEDELEMENTFRAMEWORK_API FMetaDataEntryView FindForColumn(
-			TWeakObjectPtr<const UScriptStruct> Column, FName AttributeName, ESearchScope Scope = ESearchScope::ColumnOnly) const;
-		template<typename Column>
-		FMetaDataEntryView FindForColumn(FName AttributeName, ESearchScope Scope = ESearchScope::ColumnOnly) const;
+			TWeakObjectPtr<const UScriptStruct> Column, FName AttributeName) const override;
 
 	private:
-		const TypedElementDataStorage::FQueryDescription* Query{ nullptr };
-		const FMetaData* QueryWideMetaData{ nullptr };
+		const TypedElementDataStorage::FQueryDescription& Query;
 	};
+
+	/**
+	 * Short lived view of a meta data container that wraps a list of columns.
+	 */
+	class FColumnsMetaDataView final : public FMetaDataView
+	{
+	public:
+		~FColumnsMetaDataView() override = default;
+
+		TYPEDELEMENTFRAMEWORK_API explicit FColumnsMetaDataView(TConstArrayView<TWeakObjectPtr<const UScriptStruct>> InColumns);
+		TYPEDELEMENTFRAMEWORK_API FMetaDataEntryView FindForColumn(
+			TWeakObjectPtr<const UScriptStruct> Column, FName AttributeName) const override;
+
+	private:
+		TConstArrayView<TWeakObjectPtr<const UScriptStruct>> Columns;
+	};
+
+	/**
+	 * Short lived view of a meta data container that wraps generic meta data.
+	 */
+	class FGenericMetaDataView final : public FMetaDataView
+	{
+	public:
+		~FGenericMetaDataView() override = default;
+
+		TYPEDELEMENTFRAMEWORK_API explicit FGenericMetaDataView(const FMetaData& InMetaData);
+		TYPEDELEMENTFRAMEWORK_API FMetaDataEntryView FindGeneric(FName AttributeName) const override;
+
+	private:
+		const FMetaData& MetaData;
+	};
+
+	/**
+	 * Short lived view of a meta data container that wraps around another meta data view so they can be chained.
+	 */
+	class FForwardingMetaDataView final : public FMetaDataView
+	{
+	public:
+		~FForwardingMetaDataView() override = default;
+
+		TYPEDELEMENTFRAMEWORK_API explicit FForwardingMetaDataView(const FMetaDataView& InView);
+		TYPEDELEMENTFRAMEWORK_API FMetaDataEntryView FindGeneric(FName AttributeName) const override;
+		TYPEDELEMENTFRAMEWORK_API FMetaDataEntryView FindForColumn(
+			TWeakObjectPtr<const UScriptStruct> Column, FName AttributeName) const override;
+
+	private:
+		const FMetaDataView& View;
+	};
+
+	/**
+	 * Short lived view of a meta data container.
+	 */
+	template<typename... ViewTypes>
+	class FComboMetaDataView final : public FMetaDataView
+	{
+	public:
+		FComboMetaDataView(const ViewTypes&... InViews);
+		FComboMetaDataView(ViewTypes&&... InViews);
+
+		template<typename NextViewType>
+		FComboMetaDataView<ViewTypes..., NextViewType> Next(NextViewType&& NextView);
+
+		FMetaDataEntryView FindGeneric(FName AttributeName) const override;
+		FMetaDataEntryView FindForColumn(TWeakObjectPtr<const UScriptStruct> Column, FName AttributeName) const;
+
+		TTuple<ViewTypes...> Views;
+	};
+
 } // TypedElementDataStorage
 
 #include "Elements/Framework/TypedElementMetaData.inl"
\ No newline at end of file
diff --git a/Engine/Source/Runtime/TypedElementFramework/Public/Elements/Framework/TypedElementMetaData.inl b/Engine/Source/Runtime/TypedElementFramework/Public/Elements/Framework/TypedElementMetaData.inl
--- a/Engine/Source/Runtime/TypedElementFramework/Public/Elements/Framework/TypedElementMetaData.inl
+++ b/Engine/Source/Runtime/TypedElementFramework/Public/Elements/Framework/TypedElementMetaData.inl
@@ -99,8 +99,68 @@ namespace TypedElementDataStorage
 	//
 
 	template<typename Column>
-	FMetaDataEntryView FMetaDataView::FindForColumn(FName AttributeName, ESearchScope Scope) const
+	FMetaDataEntryView FMetaDataView::FindForColumn(FName AttributeName) const
 	{
-		return FindForColumn(Column::StaticStruct(), AttributeName, Scope);
+		return FindForColumn(Column::StaticStruct(), AttributeName);
+	}
+
+
+
+	//
+	// FComboMetaDataView
+	//
+
+	template<typename... ViewTypes>
+	FComboMetaDataView<ViewTypes...>::FComboMetaDataView(const ViewTypes&... InViews)
+		: Views(InViews...)
+	{
+	}
+
+	template<typename... ViewTypes>
+	FComboMetaDataView<ViewTypes...>::FComboMetaDataView(ViewTypes&&... InViews)
+		: Views(Forward<ViewTypes>(InViews)...)
+	{
+	}
+
+	template<typename... ViewTypes>
+	template<typename NextViewType>
+	FComboMetaDataView<ViewTypes..., NextViewType> FComboMetaDataView<ViewTypes...>::Next(NextViewType&& NextView)
+	{
+		return Views.ApplyBefore([](const ViewTypes&... InViews, const NextViewType& InNextView)
+			{
+				return FComboMetaDataView<ViewTypes..., NextViewType>(InViews..., InNextView);
+			}, Forward<NextViewType>(NextView));
+	}
+
+	template<typename... ViewTypes>
+	FMetaDataEntryView FComboMetaDataView<ViewTypes...>::FindGeneric(FName AttributeName) const
+	{
+		return Views.ApplyAfter([&AttributeName](const ViewTypes&... InViews)
+			{
+				FMetaDataEntryView Result;
+				auto Run = [](FName AttributeName, FMetaDataEntryView& Result, const auto& View)
+				{
+					Result = View.FindGeneric(AttributeName);
+					return Result.IsSet();
+				};
+				(Run(AttributeName, Result, InViews) || ...);
+				return Result;
+			});
+	}
+
+	template<typename... ViewTypes>
+	FMetaDataEntryView FComboMetaDataView<ViewTypes...>::FindForColumn(TWeakObjectPtr<const UScriptStruct> Column, FName AttributeName) const
+	{
+		return Views.ApplyAfter([&Column, &AttributeName](const ViewTypes&... InViews)
+			{
+				FMetaDataEntryView Result;
+				auto Run = [](TWeakObjectPtr<const UScriptStruct> Column, FName AttributeName, FMetaDataEntryView& Result, const auto& View)
+				{
+					Result = View.FindForColumn(Column, AttributeName);
+					return Result.IsSet();
+				};
+				(Run(Column, AttributeName, Result, InViews) || ...);
+				return Result;
+			});
 	}
 } // namespace TypedElementDataStorage
diff --git a/Engine/Source/Runtime/TypedElementFramework/Public/Elements/Interfaces/TypedElementDataStorageUiInterface.h b/Engine/Source/Runtime/TypedElementFramework/Public/Elements/Interfaces/TypedElementDataStorageUiInterface.h
--- a/Engine/Source/Runtime/TypedElementFramework/Public/Elements/Interfaces/TypedElementDataStorageUiInterface.h
+++ b/Engine/Source/Runtime/TypedElementFramework/Public/Elements/Interfaces/TypedElementDataStorageUiInterface.h
@@ -53,12 +53,12 @@ public:
 		TypedElementRowHandle Row, /** The row the widget will be stored in. */
 		ITypedElementDataStorageInterface* DataStorage,
 		ITypedElementDataStorageUiInterface* DataStorageUi,
-		TypedElementDataStorage::FMetaDataView Arguments);
+		const TypedElementDataStorage::FMetaDataView& Arguments);
 
 protected:
 	/** Create a new instance of the target widget. This is a required function. */
 	TYPEDELEMENTFRAMEWORK_API virtual TSharedPtr<SWidget> CreateWidget(
-		TypedElementDataStorage::FMetaDataView Arguments) PURE_VIRTUAL(FTypedElementWidgetConstructor::CreateWidget, return nullptr; );
+		const TypedElementDataStorage::FMetaDataView& Arguments) PURE_VIRTUAL(FTypedElementWidgetConstructor::CreateWidget, return nullptr; );
 	/** Set any values in columns if needed. The columns provided through GetAdditionalColumnsList() will have already been created. */
 	TYPEDELEMENTFRAMEWORK_API virtual bool SetColumns(ITypedElementDataStorageInterface* DataStorage, TypedElementRowHandle Row);
 	/** 
@@ -192,25 +192,25 @@ public:
 	
 	/** Creates widget constructors for the requested purpose. */
 	virtual void CreateWidgetConstructors(FName Purpose, 
-		TypedElementDataStorage::FMetaDataView Arguments, const WidgetConstructorCallback& Callback) = 0;
+		const TypedElementDataStorage::FMetaDataView& Arguments, const WidgetConstructorCallback& Callback) = 0;
 	/** 
 	 * Finds matching widget constructors for provided columns, preferring longer matches over shorter matches.
 	 * The provided list of columns will be updated to contain all columns that couldn't be matched.
 	 */
 	virtual void CreateWidgetConstructors(FName Purpose, EMatchApproach MatchApproach, TArray<TWeakObjectPtr<const UScriptStruct>>& Columns,
-		TypedElementDataStorage::FMetaDataView Arguments, const WidgetConstructorCallback& Callback) = 0;
+		const TypedElementDataStorage::FMetaDataView& Arguments, const WidgetConstructorCallback& Callback) = 0;
 
 	/**
 	 * Creates all the widgets registered under the provided name. This may be a large number of widgets for a wide name
 	 * or exactly one when the exact name of the widget is registered. Arguments can be provided, but widgets are free
 	 * to ignore them.
 	 */
-	virtual void ConstructWidgets(FName Purpose, TypedElementDataStorage::FMetaDataView Arguments,
+	virtual void ConstructWidgets(FName Purpose, const TypedElementDataStorage::FMetaDataView& Arguments,
 		const WidgetCreatedCallback& ConstructionCallback) = 0;
 
 	/** Creates a single widget using the provided constructor. Arguments can optionally be used to intialize the constructor. */
 	virtual TSharedPtr<SWidget> ConstructWidget(TypedElementRowHandle Row, FTypedElementWidgetConstructor& Constructor,
-		TypedElementDataStorage::FMetaDataView Arguments) = 0;
+		const TypedElementDataStorage::FMetaDataView& Arguments) = 0;
 
 	/** Calls the provided callback for all known registered widget purposes. */
 	virtual void ListWidgetPurposes(const WidgetPurposeCallback& Callback) const = 0;
