Add revision control columns initial functionality

diff --git a/Engine/Plugins/Experimental/TypedElementsDataStorage/Source/TypedElementsDataStorage/Private/Compatibility/TypedElementSCCProcessors.cpp b/Engine/Plugins/Experimental/TypedElementsDataStorage/Source/TypedElementsDataStorage/Private/Compatibility/TypedElementSCCProcessors.cpp
new file mode 100644
--- /dev/null
+++ b/Engine/Plugins/Experimental/TypedElementsDataStorage/Source/TypedElementsDataStorage/Private/Compatibility/TypedElementSCCProcessors.cpp
@@ -0,0 +1,93 @@
+﻿// Copyright Epic Games, Inc. All Rights Reserved.
+
+#include "TypedElementSCCProcessors.h"
+
+#include "Elements/Columns/TypedElementMiscColumns.h"
+#include "Elements/Columns/TypedElementRevisionControlColumns.h"
+#include "Elements/Columns/TypedElementSelectionColumns.h"
+#include "Elements/Columns/TypedElementViewportColumns.h"
+#include "Elements/Framework/TypedElementQueryBuilder.h"
+#include "Elements/Framework/TypedElementRegistry.h"
+
+FAutoConsoleCommandWithArgsAndOutputDevice SetSelectionSCCStateConsoleCommand(
+	TEXT("TEDS.Debug.SetSCCState"),
+	TEXT("Adds a source control state to selected objects."),
+	FConsoleCommandWithArgsAndOutputDeviceDelegate::CreateLambda([](const TArray<FString>& Args, FOutputDevice& Output)
+		{
+			using namespace TypedElementQueryBuilder;
+			using DSI = ITypedElementDataStorageInterface;
+
+			TRACE_CPUPROFILER_EVENT_SCOPE(TEDS.Debug.SetSCCState);
+
+			if (ITypedElementDataStorageInterface* DataStorage = UTypedElementRegistry::GetInstance()->GetMutableDataStorage())
+			{
+				static TypedElementQueryHandle AllSelectedQuery = TypedElementInvalidQueryHandle;
+				if (AllSelectedQuery == TypedElementInvalidQueryHandle)
+				{
+					AllSelectedQuery = DataStorage->RegisterQuery(
+						Select()
+						.Where()
+							.All<FTypedElementSelectionColumn>()
+						.Compile());
+				}
+				
+				if (AllSelectedQuery == TypedElementInvalidQueryHandle)
+				{
+					return;
+				}
+
+				ESCCModification Modification = ESCCModification::Modified;
+
+				if (Args.Num() > 0)
+				{
+					// Parse the index
+					int32 StateIndex;
+					LexFromString(StateIndex, *Args[0]);
+
+					if (!(StateIndex >= 0 && StateIndex <= int32(ESCCModification::Conflicted)))
+					{
+						Output.Log(TEXT("State index out of range"));
+						return;
+					}
+
+					Modification = static_cast<ESCCModification>(StateIndex);
+				}
+				
+				TArray<TypedElementRowHandle> RowHandles;
+				
+				DataStorage->RunQuery(AllSelectedQuery, [&RowHandles](const DSI::FQueryDescription&, DSI::IDirectQueryContext& Context)
+				{
+					RowHandles = Context.GetRowHandles();
+				});
+
+				for (TypedElementRowHandle Row : RowHandles)
+				{
+					DataStorage->AddOrGetColumn<FSCCStatusColumn>(Row)->Modification = Modification;
+					DataStorage->AddColumn<FTypedElementSyncFromWorldTag>(Row);
+				}
+			}
+		}
+	));
+
+void UTypedElementSCCFactory::RegisterQueries(ITypedElementDataStorageInterface& DataStorage) const
+{
+	using namespace TypedElementQueryBuilder;
+	using DSI = ITypedElementDataStorageInterface;
+
+	DataStorage.RegisterQuery(
+		Select(
+			TEXT("Change selection outline colors based on SCC status"),
+			// This is in PrePhysics because the outline->actor query is in DuringPhysics and contexts don't flush changes between tick groups
+			FProcessor(DSI::EQueryTickPhase::PrePhysics, DataStorage.GetQueryTickGroupName(DSI::EQueryTickGroups::SyncExternalToDataStorage)),
+			[](DSI::IQueryContext& Context, TypedElementRowHandle Row, const FSCCStatusColumn& SCCStatus)
+			{
+				constexpr int BasicSelectionColorCount = 2;
+				Context.AddColumn<FTypedElementViewportColorColumn>(Row, { .SelectionOutlineColorIndex = static_cast<uint8>(int(SCCStatus.Modification) + BasicSelectionColorCount) });
+				Context.AddColumns<FTypedElementSyncBackToWorldTag>(Row);
+			}
+		)
+		.Where()
+			.All<FTypedElementSyncFromWorldTag>()
+		.Compile()
+	);
+}
diff --git a/Engine/Plugins/Experimental/TypedElementsDataStorage/Source/TypedElementsDataStorage/Private/Compatibility/TypedElementSCCProcessors.h b/Engine/Plugins/Experimental/TypedElementsDataStorage/Source/TypedElementsDataStorage/Private/Compatibility/TypedElementSCCProcessors.h
new file mode 100644
--- /dev/null
+++ b/Engine/Plugins/Experimental/TypedElementsDataStorage/Source/TypedElementsDataStorage/Private/Compatibility/TypedElementSCCProcessors.h
@@ -0,0 +1,20 @@
+﻿// Copyright Epic Games, Inc. All Rights Reserved.
+
+#pragma once
+
+#include "Elements/Interfaces/TypedElementDataStorageFactory.h"
+#include "UObject/ObjectMacros.h"
+
+#include "TypedElementSCCProcessors.generated.h"
+
+UCLASS()
+class UTypedElementSCCFactory : public UTypedElementDataStorageFactory
+{
+	GENERATED_BODY()
+
+public:
+	~UTypedElementSCCFactory() override = default;
+
+	void RegisterQueries(ITypedElementDataStorageInterface& DataStorage) const override;
+	
+};
diff --git a/Engine/Source/Runtime/TypedElementFramework/Public/Elements/Columns/TypedElementRevisionControlColumns.h b/Engine/Source/Runtime/TypedElementFramework/Public/Elements/Columns/TypedElementRevisionControlColumns.h
new file mode 100644
--- /dev/null
+++ b/Engine/Source/Runtime/TypedElementFramework/Public/Elements/Columns/TypedElementRevisionControlColumns.h
@@ -0,0 +1,94 @@
+﻿// Copyright Epic Games, Inc. All Rights Reserved.
+
+#pragma once
+
+#include "Elements/Interfaces/TypedElementDataStorageInterface.h"
+
+#include "TypedElementRevisionControlColumns.generated.h"
+
+USTRUCT(meta = (DisplayName = "In a changelist"))
+struct FSCCInChangelistTag final : public FTypedElementDataStorageTag
+{
+	GENERATED_BODY()
+};
+
+USTRUCT(meta = (DisplayName = "Staged"))
+struct FSCCStagedTag final : public FTypedElementDataStorageTag
+{
+	GENERATED_BODY()
+};
+
+USTRUCT(meta = (DisplayName = "Locked by you"))
+struct FSCCLockedTag final : public FTypedElementDataStorageTag
+{
+	GENERATED_BODY()
+};
+
+USTRUCT(meta = (DisplayName = "Edited by others"))
+struct FSCCExternallyEditedTag final : public FTypedElementDataStorageTag
+{
+	GENERATED_BODY()
+};
+
+UENUM()
+enum class ESCCModification
+{
+	Modified,
+	Added,
+	Removed,
+	Conflicted,
+};
+
+USTRUCT(meta = (DisplayName = "Revision Control Status"))
+struct FSCCStatusColumn final : public FTypedElementDataStorageColumn
+{
+	GENERATED_BODY()
+
+	UPROPERTY()
+	ESCCModification Modification;
+};
+
+USTRUCT()
+struct FSCCRevisionId
+{
+	GENERATED_BODY()
+
+	UPROPERTY()
+	uint32 Id[5];
+};
+
+USTRUCT(meta = (DisplayName = "Revision ID"))
+struct FSCCRevisionIdColumn final : public FTypedElementDataStorageColumn
+{
+	GENERATED_BODY()
+
+	UPROPERTY()
+	FSCCRevisionId RevisionId; 
+};
+
+USTRUCT(meta = (DisplayName = "Revision ID from server"))
+struct FSCCExternalRevisionIdColumn final : public FTypedElementDataStorageColumn
+{
+	GENERATED_BODY()
+
+	UPROPERTY()
+	FSCCRevisionId RevisionId; 
+};
+
+USTRUCT()
+struct FSCCUserInfo
+{
+	GENERATED_BODY()
+
+	UPROPERTY()
+	FString Name;
+};
+
+USTRUCT(meta = (DisplayName = "Locked by others"))
+struct FSCCExternallyLockedColumn final : public FTypedElementDataStorageColumn
+{
+	GENERATED_BODY()
+
+	UPROPERTY()
+	FSCCUserInfo LockedBy;
+};
