TEDS: Add a widget constructor for the ClassTypeInfo column in TEDS that can create a text or icon widget.

diff --git a/Engine/Plugins/Experimental/TypedElementsDataStorage/Source/TypedElementsDataStorage/Private/TypedElementDatabaseUI.cpp b/Engine/Plugins/Experimental/TypedElementsDataStorage/Source/TypedElementsDataStorage/Private/TypedElementDatabaseUI.cpp
--- a/Engine/Plugins/Experimental/TypedElementsDataStorage/Source/TypedElementsDataStorage/Private/TypedElementDatabaseUI.cpp
+++ b/Engine/Plugins/Experimental/TypedElementsDataStorage/Source/TypedElementsDataStorage/Private/TypedElementDatabaseUI.cpp
@@ -599,12 +599,8 @@ void UTypedElementDatabaseUi::CreateWidgetConstructors_SingleMatch(const TArray<
 	auto FactoryEnd = WidgetFactories.rend();
 
 	// Start from the back as the widgets with lower counts will be last.
-	int32 ColumnIndex = Columns.Num() - 1;
-
-	auto ColumnEnd = Columns.rend();
-	for (auto ColumnIt = Columns.rbegin(); ColumnIt != ColumnEnd; ++ColumnIt, --ColumnIndex)
+	for (int32 ColumnIndex = Columns.Num() - 1; ColumnIndex >= 0; --ColumnIndex)
 	{
-		const UScriptStruct* ColumnType = (*ColumnIt).Get();
 		for (; FactoryIt != FactoryEnd; ++FactoryIt)
 		{
 			TConstArrayView<TWeakObjectPtr<const UScriptStruct>> ColumnData = (*FactoryIt).Columns.GetColumns();
@@ -619,7 +615,7 @@ void UTypedElementDatabaseUi::CreateWidgetConstructors_SingleMatch(const TArray<
 				continue;
 			}
 
-			if (ColumnData[0] == *ColumnIt)
+			if (ColumnData[0] == Columns[ColumnIndex])
 			{
 				Columns.RemoveAt(ColumnIndex);
 				CreateSingleWidgetConstructor((*FactoryIt).Constructor, Arguments, ColumnData, Callback);
diff --git a/Engine/Plugins/Experimental/TypedElementsDataStorage/Source/TypedElementsDataStorageUI/Private/TypedElementTypeInfoWidget.cpp b/Engine/Plugins/Experimental/TypedElementsDataStorage/Source/TypedElementsDataStorageUI/Private/TypedElementTypeInfoWidget.cpp
new file mode 100644
--- /dev/null
+++ b/Engine/Plugins/Experimental/TypedElementsDataStorage/Source/TypedElementsDataStorageUI/Private/TypedElementTypeInfoWidget.cpp
@@ -0,0 +1,138 @@
+// Copyright Epic Games, Inc. All Rights Reserved.
+
+#include "TypedElementTypeInfoWidget.h"
+
+#include "MassActorSubsystem.h"
+#include "Elements/Columns/TypedElementMiscColumns.h"
+#include "Elements/Columns/TypedElementTypeInfoColumns.h"
+#include "Styling/SlateIconFinder.h"
+#include "Widgets/Text/STextBlock.h"
+#include "Widgets/Images/SImage.h"
+
+void UTypedElementTypeInfoWidgetFactory::RegisterWidgetConstructors(ITypedElementDataStorageInterface& DataStorage,
+	ITypedElementDataStorageUiInterface& DataStorageUi) const
+{
+	DataStorageUi.RegisterWidgetFactory<FTypedElementTypeInfoWidgetConstructor>(FName(TEXT("General.Cell")),
+		TypedElementQueryBuilder::FColumn<FTypedElementClassTypeInfoColumn>());
+
+}
+
+TMap<FName, const FSlateBrush*> FTypedElementTypeInfoWidgetConstructor::CachedIconMap;
+
+FTypedElementTypeInfoWidgetConstructor::FTypedElementTypeInfoWidgetConstructor()
+	: Super(FTypedElementTypeInfoWidgetConstructor::StaticStruct())
+	, bUseIcon(false)
+{
+}
+
+FTypedElementTypeInfoWidgetConstructor::FTypedElementTypeInfoWidgetConstructor(const UScriptStruct* InTypeInfo)
+	: Super(InTypeInfo)
+	, bUseIcon(false)
+{
+	
+}
+
+TSharedPtr<SWidget> FTypedElementTypeInfoWidgetConstructor::CreateWidget(
+	const TypedElementDataStorage::FMetaDataView& Arguments	)
+{
+	// Check if the caller provided metadata to use an icon widget
+	TypedElementDataStorage::FMetaDataEntryView MetaDataEntryView = Arguments.FindGeneric("TypedElementTypeInfoWidget_bUseIcon");
+	if(MetaDataEntryView.IsSet())
+	{
+		check(MetaDataEntryView.IsType<bool>());
+
+		bUseIcon = *MetaDataEntryView.TryGetExact<bool>();
+	}
+
+	// Can't do this in one line because slate doesn't like multiple SNew in one line
+	TSharedPtr<SWidget> TypeWidget;
+	if(bUseIcon)
+	{
+		TypeWidget = SNew(SImage)
+					.DesiredSizeOverride(FVector2D(16.f, 16.f))
+					.ColorAndOpacity(FSlateColor::UseForeground());
+	}
+	else
+	{
+		TypeWidget = SNew(STextBlock);
+	}
+	
+	return TypeWidget;
+}
+
+bool FTypedElementTypeInfoWidgetConstructor::FinalizeWidget(ITypedElementDataStorageInterface* DataStorage,
+	ITypedElementDataStorageUiInterface* DataStorageUi, TypedElementRowHandle Row, const TSharedPtr<SWidget>& Widget)
+{
+	checkf(Widget, TEXT("Referenced widget is not valid. A constructed widget may not have been cleaned up. This can "
+		"also happen if this processor is running in the same phase as the processors responsible for cleaning up old "
+		"references."));
+
+	TypedElementRowHandle TargetRow = DataStorage->GetColumn<FTypedElementRowReferenceColumn>(Row)->Row;
+
+	if (const FTypedElementClassTypeInfoColumn* TypeInfoColumn = DataStorage->GetColumn<FTypedElementClassTypeInfoColumn>(TargetRow))
+	{
+		if(bUseIcon)
+		{
+			checkf(Widget->GetType() == SImage::StaticWidgetClass().GetWidgetType(),
+				TEXT("Stored widget with FTypedElementLabelWidgetConstructor doesn't match type %s, but was a %s."),
+				*(SImage::StaticWidgetClass().GetWidgetType().ToString()),
+				*(Widget->GetTypeAsString()));
+
+			SImage* WidgetInstance = static_cast<SImage*>(Widget.Get());
+			
+			WidgetInstance->SetImage(GetIconForRow(DataStorage, Row, TypeInfoColumn));
+		}
+		else
+		{
+			checkf(Widget->GetType() == STextBlock::StaticWidgetClass().GetWidgetType(),
+				TEXT("Stored widget with FTypedElementLabelWidgetConstructor doesn't match type %s, but was a %s."),
+				*(STextBlock::StaticWidgetClass().GetWidgetType().ToString()),
+				*(Widget->GetTypeAsString()));
+				
+			STextBlock* WidgetInstance = static_cast<STextBlock*>(Widget.Get());
+
+			WidgetInstance->SetText(FText::FromString(TypeInfoColumn->TypeInfo.Get()->GetName()));
+		}
+	}
+
+	return true;
+}
+
+const FSlateBrush* FTypedElementTypeInfoWidgetConstructor::GetIconForRow(ITypedElementDataStorageInterface* DataStorage, TypedElementRowHandle Row, const FTypedElementClassTypeInfoColumn* TypeInfoColumn)
+{
+
+	/* The logic here is very similar to SActorTreeLabel::GetIcon in ActorTreeItem.cpp which allows the actor to specify
+	 * an override for the icon, and has a fallback to the class icon if not.
+	 */
+	
+	const UClass* Type = TypeInfoColumn->TypeInfo.Get();
+	FName IconName = Type->GetFName();
+
+	// Allow the actor the first chance to provide an icon override
+	FMassActorFragment* ActorStore = DataStorage->GetColumn<FMassActorFragment>(Row);
+	if(ActorStore)
+	{
+		if(const AActor* Actor = ActorStore->Get())
+		{
+			IconName = Actor->GetCustomIconName();
+		}
+	}
+
+	// Check the cache if we already found an icon for this class
+	if(const FSlateBrush** CachedBrush = CachedIconMap.Find(IconName))
+	{
+		if(*CachedBrush)
+		{
+			return *CachedBrush;
+		}
+	}
+
+	if(const FSlateBrush* FoundSlateBrush = FSlateIconFinder::FindIconForClass(Type).GetOptionalIcon())
+	{
+		CachedIconMap.Add(IconName, FoundSlateBrush);
+		return FoundSlateBrush;
+	}
+
+	// Fallback to the regular actor icon if we haven't found any specific icon
+	return FSlateIconFinder::FindIconForClass(AActor::StaticClass()).GetOptionalIcon();;
+}
\ No newline at end of file
diff --git a/Engine/Plugins/Experimental/TypedElementsDataStorage/Source/TypedElementsDataStorageUI/Public/TypedElementTypeInfoWidget.h b/Engine/Plugins/Experimental/TypedElementsDataStorage/Source/TypedElementsDataStorageUI/Public/TypedElementTypeInfoWidget.h
new file mode 100644
--- /dev/null
+++ b/Engine/Plugins/Experimental/TypedElementsDataStorage/Source/TypedElementsDataStorageUI/Public/TypedElementTypeInfoWidget.h
@@ -0,0 +1,53 @@
+// Copyright Epic Games, Inc. All Rights Reserved.
+
+#pragma once
+
+#include "Elements/Interfaces/TypedElementDataStorageFactory.h"
+#include "Elements/Interfaces/TypedElementDataStorageInterface.h"
+#include "Elements/Interfaces/TypedElementDataStorageUiInterface.h"
+
+#include "TypedElementTypeInfoWidget.generated.h"
+
+struct FSlateBrush;
+struct FTypedElementClassTypeInfoColumn;
+
+UCLASS()
+class TYPEDELEMENTSDATASTORAGEUI_API UTypedElementTypeInfoWidgetFactory : public UTypedElementDataStorageFactory
+{
+	GENERATED_BODY()
+
+public:
+	~UTypedElementTypeInfoWidgetFactory() override = default;
+
+	void RegisterWidgetConstructors(ITypedElementDataStorageInterface& DataStorage,
+		ITypedElementDataStorageUiInterface& DataStorageUi) const override;
+};
+
+USTRUCT()
+struct TYPEDELEMENTSDATASTORAGEUI_API FTypedElementTypeInfoWidgetConstructor : public FTypedElementWidgetConstructor
+{
+	GENERATED_BODY()
+
+public:
+	FTypedElementTypeInfoWidgetConstructor();
+	~FTypedElementTypeInfoWidgetConstructor() override = default;
+
+protected:
+	explicit FTypedElementTypeInfoWidgetConstructor(const UScriptStruct* InTypeInfo);
+	TSharedPtr<SWidget> CreateWidget(const TypedElementDataStorage::FMetaDataView& Arguments) override;
+	bool FinalizeWidget(ITypedElementDataStorageInterface* DataStorage, ITypedElementDataStorageUiInterface* DataStorageUi,
+		TypedElementRowHandle Row, const TSharedPtr<SWidget>& Widget) override;
+
+private:
+
+	// Get the icon for a given row, checking the cache to see if we already have one
+	static const FSlateBrush* GetIconForRow(ITypedElementDataStorageInterface* DataStorage, TypedElementRowHandle Row, const FTypedElementClassTypeInfoColumn* TypeInfoColumn);
+	
+private:
+
+	// Cache to avoid looking up the icon for a class every time
+	static TMap<FName, const FSlateBrush*> CachedIconMap;
+
+	// Whether the widget created by this constructor should be icon or text
+	bool bUseIcon;
+};
diff --git a/Engine/Source/Runtime/TypedElementFramework/Public/Elements/Framework/TypedElementMetaData.inl b/Engine/Source/Runtime/TypedElementFramework/Public/Elements/Framework/TypedElementMetaData.inl
--- a/Engine/Source/Runtime/TypedElementFramework/Public/Elements/Framework/TypedElementMetaData.inl
+++ b/Engine/Source/Runtime/TypedElementFramework/Public/Elements/Framework/TypedElementMetaData.inl
@@ -11,7 +11,11 @@ namespace TypedElementDataStorage
 
 			if constexpr (std::is_arithmetic_v<BaseType>)
 			{
-				if constexpr (std::is_integral_v<BaseType>)
+				if constexpr (std::is_same_v<bool, BaseType>)
+				{
+					Container.Add(Name, MetaDataType(TInPlaceType<bool>(), static_cast<bool>(Value)));
+				}
+				else if constexpr (std::is_integral_v<BaseType>)
 				{
 					if constexpr (std::is_signed_v<BaseType>)
 					{
@@ -26,10 +30,6 @@ namespace TypedElementDataStorage
 				{
 					Container.Add(Name, MetaDataType(TInPlaceType<double>(), static_cast<double>(Value)));
 				}
-				else
-				{
-					Container.Add(Name, MetaDataType(TInPlaceType<bool>(), static_cast<bool>(Value)));
-				}
 			}
 			else if constexpr (std::is_convertible_v<BaseType, FString>)
 			{
