Enable selection column reinstancing

* Memento system translators from Mementos to Columns changed to add the column and initialize it. It is also not required to use the subquery api since the operation is write-only to the database and therefore no need to lock the table.
* AddColumn for IQueryContext changed to add columns only if one does not exist.  This avoids the run time exception Mass generates if trying to add duplicate columns and is a step closer to making the API idempotent.
* Fixed spelling mistake in AddColumnUninitialized and changed the documentation.
* Removed the examples and replaced with a translator enable the system to automagically mementoize the selection column.  Selection columns will now be added and initialized to reinstanced objects.

diff --git a/Engine/Plugins/Experimental/TypedElementsDataStorage/Source/TypedElementsDataStorage/Private/Memento/TypedElementMementoExamples.h b/Engine/Plugins/Experimental/TypedElementsDataStorage/Source/TypedElementsDataStorage/Private/Memento/TypedElementMementoExamples.h
deleted file mode 100644
diff --git a/Engine/Plugins/Experimental/TypedElementsDataStorage/Source/TypedElementsDataStorage/Private/Memento/TypedElementMementoSystem.cpp b/Engine/Plugins/Experimental/TypedElementsDataStorage/Source/TypedElementsDataStorage/Private/Memento/TypedElementMementoSystem.cpp
--- a/Engine/Plugins/Experimental/TypedElementsDataStorage/Source/TypedElementsDataStorage/Private/Memento/TypedElementMementoSystem.cpp
+++ b/Engine/Plugins/Experimental/TypedElementsDataStorage/Source/TypedElementsDataStorage/Private/Memento/TypedElementMementoSystem.cpp
@@ -1,4 +1,4 @@
-﻿// Copyright Epic Games, Inc. All Rights Reserved.
+// Copyright Epic Games, Inc. All Rights Reserved.
 
 #include "TypedElementMementoSystem.h"
 
@@ -66,18 +66,19 @@ void UTypedElementMementoSystem::RegisterQueries(UTypedElementDatabase& DataStor
 		const UScriptStruct* MementoizedColumn = MementoTranslator->GetColumnType();
 		const UScriptStruct* Memento = MementoTranslator->GetMementoType();
 		
-		const FName TranslationProcessorName = FName(FString::Printf(TEXT("MementoTranslator: %s -> %s"), *MementoizedColumn->GetName(), *Memento->GetName()));
+		const FName TranslationProcessorName = FName(FString::Printf(TEXT("MementoTranslator (Column->Memento) %s -> %s"), *MementoizedColumn->GetName(), *Memento->GetName()));
 		const TypedElementQueryHandle QueryHandle = DataStorage.RegisterQuery(
 			Select(
 				TranslationProcessorName,
 				FObserver::OnRemove<FTypedElementMementoOnDelete>(),
-				[MementoTranslator, Memento, MementoizedColumn](TypedElementDataStorage::IQueryContext& Context, TypedElementRowHandle row, const FTypedElementMementoOnDelete& MementoRow)
+				[MementoTranslator, Memento, MementoizedColumn](TypedElementDataStorage::IQueryContext& Context, TypedElementRowHandle Row, const FTypedElementMementoOnDelete& MementoRow)
 				{
-					void* MementoObject = Context.AddColumnUnitialized(MementoRow.Memento, Memento);
-					Memento->InitializeStruct(MementoObject);
+					void* StagedMementoColumn = Context.AddColumnUninitialized(MementoRow.Memento, Memento);
+					// StagedMementoColumn will be uninitialized memory, ensure constructor called
+					Memento->InitializeStruct(StagedMementoColumn);
 		
 					const void* SourceColumn = Context.GetColumn(MementoizedColumn);
-					MementoTranslator->TranslateColumnToMemento(SourceColumn, MementoObject);
+					MementoTranslator->TranslateColumnToMemento(SourceColumn, StagedMementoColumn);
 				})
 				.ReadOnly(MementoizedColumn)
 				.Compile());
@@ -107,36 +108,29 @@ void UTypedElementMementoSystem::RegisterQueries(UTypedElementDatabase& DataStor
 		const UScriptStruct* MementoizedColumnType = MementoTranslator->GetColumnType();
 		const UScriptStruct* MementoType = MementoTranslator->GetMementoType();
 
-		TypedElementQueryHandle Subquery = DataStorage.RegisterQuery(
-			Select()
-				.ReadWrite(MementoizedColumnType)
-			.Compile());
+		namespace TEDS = TypedElementDataStorage;
 
-		const FName TranslationProcessorName = FName(FString::Printf(TEXT("MementoTranslator: %s -> %s"), *MementoType->GetName(), *MementoizedColumnType->GetName()));
-		const TypedElementQueryHandle QueryHandle = DataStorage.RegisterQuery(
+		const FName TranslationProcessorName = FName(FString::Printf(TEXT("MementoTranslator (Memento->Column) %s -> %s"), *MementoType->GetName(), *MementoizedColumnType->GetName()));
+		const TypedElementQueryHandle QueryHandleForExistence = DataStorage.RegisterQuery(
 			Select(
 				TranslationProcessorName,
 				FProcessor(DSI::EQueryTickPhase::PostPhysics, DataStorage.GetQueryTickGroupName(DSI::EQueryTickGroups::Default)),
-				[MementoTranslator](TypedElementDataStorage::IQueryContext& Context, TypedElementRowHandle Row, const FTypedElementMementoReinstanceTarget& ReinstanceTarget)
+				[MementoTranslator](TEDS::IQueryContext& Context, TypedElementRowHandle Row, const FTypedElementMementoReinstanceTarget& ReinstanceTarget)
 				{
-					const UScriptStruct* MementoType = MementoTranslator->GetMementoType();
-					
-					const void* Memento = Context.GetColumn(MementoType);
-					Context.RunSubquery(0, ReinstanceTarget.Target, 
-						[MementoTranslator, Memento](const DSI::FQueryDescription&, DSI::ISubqueryContext& SubQueryContext)
-						{
-							const UScriptStruct* MementoizedColumnType = MementoTranslator->GetColumnType();
-							void* MementoizedColumn = SubQueryContext.GetMutableColumn(MementoizedColumnType);
-
-							MementoTranslator->TranslateMementoToColumn(Memento, MementoizedColumn);
-						});
+					// Add the column, note the column may already exist.
+					void* StagedColumn = Context.AddColumnUninitialized(ReinstanceTarget.Target, MementoTranslator->GetColumnType());
+
+					// StagedColumn will be uninitialized memory, ensure constructor called
+					MementoTranslator->GetColumnType()->InitializeStruct(StagedColumn);
+
+					const void* Memento = Context.GetColumn(MementoTranslator->GetMementoType());
+					MementoTranslator->TranslateMementoToColumn(Memento, StagedColumn);
 				})
 				.ReadOnly(MementoType)
 				.Where()
 					.All<FTypedElementMementoTag>()
-				.DependsOn().SubQuery(Subquery)
 				.Compile());
-		check(QueryHandle != TypedElementInvalidQueryHandle);
+		check(QueryHandleForExistence != TypedElementInvalidQueryHandle);
 	}
 
 	/**
diff --git a/Engine/Plugins/Experimental/TypedElementsDataStorage/Source/TypedElementsDataStorage/Private/Memento/TypedElementMementoizedColumns.h b/Engine/Plugins/Experimental/TypedElementsDataStorage/Source/TypedElementsDataStorage/Private/Memento/TypedElementMementoizedColumns.h
new file mode 100644
--- /dev/null
+++ b/Engine/Plugins/Experimental/TypedElementsDataStorage/Source/TypedElementsDataStorage/Private/Memento/TypedElementMementoizedColumns.h
@@ -0,0 +1,18 @@
+﻿// Copyright Epic Games, Inc. All Rights Reserved.
+
+#pragma once
+
+#include "Elements/Columns/TypedElementSelectionColumns.h"
+#include "Memento/TypedElementMementoTranslators.h"
+#include "TypedElementMementoizedColumns.generated.h"
+
+/**
+ * Enable SelectionColumn to be mementoized
+ */
+UCLASS()
+class UTypedElementSelectionColumnMementoTranslator final : public UTypedElementDefaultMementoTranslator
+{
+	GENERATED_BODY()
+public:
+	virtual const UScriptStruct* GetColumnType() const override { return FTypedElementSelectionColumn::StaticStruct(); }
+};
\ No newline at end of file
diff --git a/Engine/Plugins/Experimental/TypedElementsDataStorage/Source/TypedElementsDataStorage/Private/Processors/TypedElementProcessorAdaptors.cpp b/Engine/Plugins/Experimental/TypedElementsDataStorage/Source/TypedElementsDataStorage/Private/Processors/TypedElementProcessorAdaptors.cpp
--- a/Engine/Plugins/Experimental/TypedElementsDataStorage/Source/TypedElementsDataStorage/Private/Processors/TypedElementProcessorAdaptors.cpp
+++ b/Engine/Plugins/Experimental/TypedElementsDataStorage/Source/TypedElementsDataStorage/Private/Processors/TypedElementProcessorAdaptors.cpp
@@ -233,7 +233,7 @@ public:
 			TConstArrayView<FMassEntityHandle>(reinterpret_cast<const FMassEntityHandle*>(Rows.begin()), Rows.Num()));
 	}
 
-	void* AddColumnUnitialized(TypedElementDataStorage::RowHandle Row, const UScriptStruct* ObjectType) override
+	void* AddColumnUninitialized(TypedElementDataStorage::RowHandle Row, const UScriptStruct* ObjectType) override
 	{
 		checkf(ObjectType->IsChildOf(FMassTag::StaticStruct()) || ObjectType->IsChildOf(FMassFragment::StaticStruct()), TEXT("Column [%s] must be a Mass Fragment or Tag"), *ObjectType->GetName());
 		
@@ -279,7 +279,13 @@ public:
 			Context.Defer().PushCommand<FMassDeferredAddCommand>(
 				[AddedColumn](FMassEntityManager& System)
 				{
-					System.AddFragmentToEntity(AddedColumn->Entity, AddedColumn->FragmentType);
+					// Check before adding.  Mass's AddFragmentToEntity is not idempotent and will assert if adding
+					// column to a row that already has one
+					FStructView Fragment = System.GetFragmentDataStruct(AddedColumn->Entity, AddedColumn->FragmentType);
+					if (!Fragment.IsValid())
+					{
+						System.AddFragmentToEntity(AddedColumn->Entity, AddedColumn->FragmentType);
+					}
 				});
 			
 			Context.Defer().PushCommand<FMassDeferredSetCommand>(
@@ -301,7 +307,7 @@ public:
 		return ObjectCopy;
 	}
 
-	void* AddColumnUnitialized(TypedElementDataStorage::RowHandle Row, const UScriptStruct* ObjectType, ObjectMoveOperator Mover) override
+	void* AddColumnUninitialized(TypedElementDataStorage::RowHandle Row, const UScriptStruct* ObjectType, ObjectMoveOperator Mover) override
 	{
 		checkf(ObjectType->IsChildOf(FMassTag::StaticStruct()) || ObjectType->IsChildOf(FMassFragment::StaticStruct()), TEXT("Column [%s] must be a Mass Fragment or Tag"), *ObjectType->GetName());
 		
@@ -350,7 +356,13 @@ public:
 			Context.Defer().PushCommand<FMassDeferredAddCommand>(
 				[AddedColumn](FMassEntityManager& System)
 				{
-					System.AddFragmentToEntity(AddedColumn->Entity, AddedColumn->FragmentType);
+					// Check before adding.  Mass's AddFragmentToEntity is not idempotent and will assert if adding
+					// column to a row that already has one
+					FStructView Fragment = System.GetFragmentDataStruct(AddedColumn->Entity, AddedColumn->FragmentType);
+					if (!Fragment.IsValid())
+					{
+						System.AddFragmentToEntity(AddedColumn->Entity, AddedColumn->FragmentType);
+					}
 				});
 
 			Context.Defer().PushCommand<FMassDeferredSetCommand>(
diff --git a/Engine/Source/Runtime/TypedElementFramework/Public/Elements/Interfaces/TypedElementQueryStorageInterfaces.h b/Engine/Source/Runtime/TypedElementFramework/Public/Elements/Interfaces/TypedElementQueryStorageInterfaces.h
--- a/Engine/Source/Runtime/TypedElementFramework/Public/Elements/Interfaces/TypedElementQueryStorageInterfaces.h
+++ b/Engine/Source/Runtime/TypedElementFramework/Public/Elements/Interfaces/TypedElementQueryStorageInterfaces.h
@@ -110,18 +110,24 @@ namespace TypedElementDataStorage
 		 */
 		virtual void RemoveRows(TConstArrayView<RowHandle> Rows) = 0;
 		/**
-		 * Add a new unitialized column of the provided object type. The addition will not be immediately done but delayed until the
-		 * end of the tick group. The returned address is to a temporary object that's held until the column has been setup. This can be
-		 * used to set data while in the query callback but shouldn't be used outside the callback as the object will be removed.
+		 * Add a new column of the provided type if one does not exist.
+		 * Returns a staged column which is used to copy into the database at a later time via the UStructScript Copy operator
+		 * It is the caller's responsibility to ensure the staged column's constructor is called.  The caller
+		 * may modify other properties of the column.
+		 *
+		 * Note: The addition and modification of the column will not be done immediately done. Instead it will be deferred until the end of the tick group.
 		 */
-		virtual void* AddColumnUnitialized(RowHandle Row, const UScriptStruct* ObjectType) = 0;
+		virtual void* AddColumnUninitialized(RowHandle Row, const UScriptStruct* ColumnType) = 0;
+		
 		/**
-		 * Add a new unitialized column of the provided object type. The addition will not be immediately done but delayed until the
-		 * end of the tick group. The returned address is to a temporary object that's held until the column has been setup. This can be
-		 * used to set data while in the query callback but shouldn't be used outside the callback as the object will be removed. The
-		 * temporary object will be moved to it's final location using the provided move operator callback instead of being copied.
+		 * Add a new column of the provided type if one does not exist.
+		 * Returns a staged column which is used to copy into the database at a later time via the provided Move operator
+		 * It is the caller's responsibility to ensure the staged column's constructor is called.  The caller
+		 * may modify other properties of the column.
+		 *
+		 * Note: The addition and modification of the column will not be done immediately done. Instead it will be deferred until the end of the tick group.
 		 */
-		virtual void* AddColumnUnitialized(RowHandle Row, const UScriptStruct* ObjectType, ObjectMoveOperator Mover) = 0;
+		virtual void* AddColumnUninitialized(RowHandle Row, const UScriptStruct* ObjectType, ObjectMoveOperator Mover) = 0;
 		/**
 		 * Adds new empty columns to a row of the provided type. The addition will not be immediately done but delayed until the end of the
 		 * tick group.
@@ -213,7 +219,7 @@ namespace TypedElementDataStorage
 
 		if constexpr (std::is_move_constructible_v<ColumnType>)
 		{
-			void* Address = AddColumnUnitialized(Row, TypeInfo,
+			void* Address = AddColumnUninitialized(Row, TypeInfo,
 				[](void* Destination, void* Source)
 				{
 					new(Destination) ColumnType(MoveTemp(*reinterpret_cast<ColumnType*>(Source)));
@@ -222,7 +228,7 @@ namespace TypedElementDataStorage
 		}
 		else
 		{
-			void* Address = AddColumnUnitialized(Row, TypeInfo);
+			void* Address = AddColumnUninitialized(Row, TypeInfo);
 			TypeInfo->CopyScriptStruct(Address, &Column);
 			return *reinterpret_cast<ColumnType*>(Address);
 		}
