TEDS can now delay adding/removing columns if needed.

TEDS will now check if a row has been initialized or not. If the row has been created but not initialized it will not delay adding/removing columns until after the row has been initialized
This functionality is currently limited to adding columns without data. Retrieving columns in this state is still not possible.
This is an initial step in a larger future effort to have a uniform handling of commands throughout TEDS' direct calls and query callbacks.

diff --git a/Engine/Plugins/Experimental/TypedElementsDataStorage/Source/TypedElementsDataStorage/Private/TypedElementDatabase.cpp b/Engine/Plugins/Experimental/TypedElementsDataStorage/Source/TypedElementsDataStorage/Private/TypedElementDatabase.cpp
--- a/Engine/Plugins/Experimental/TypedElementsDataStorage/Source/TypedElementsDataStorage/Private/TypedElementDatabase.cpp
+++ b/Engine/Plugins/Experimental/TypedElementsDataStorage/Source/TypedElementsDataStorage/Private/TypedElementDatabase.cpp
@@ -110,6 +110,9 @@ void UTypedElementDatabase::OnPreMassTick(float DeltaTime)
 {
 	checkf(IsAvailable(), TEXT("Typed Element Database was ticked while it's not ready."));
 	OnUpdateDelegate.Broadcast();
+	// Process pending commands after other systems have had a chance to update. Other systems may have executed work needed
+	// to complete pending work.
+	ProcessPendingCommands();
 	// Recycle any full scratch memory blocks from the previous frame.
 	ScratchBuffer.RecycleBlocks();
 }
@@ -272,23 +275,17 @@ bool UTypedElementDatabase::HasRowBeenAssigned(TypedElementRowHandle Row) const
 
 bool UTypedElementDatabase::AddColumn(TypedElementRowHandle Row, const UScriptStruct* ColumnType)
 {
-	FMassEntityHandle Entity = FMassEntityHandle::FromNumber(Row);
-	if (ActiveEditorEntityManager && ActiveEditorEntityManager->IsEntityActive(Entity))
+	if (ColumnType && ActiveEditorEntityManager)
 	{
-		if (ColumnType->IsChildOf(FMassTag::StaticStruct()))
+		if (HasRowBeenAssigned(Row))
 		{
-			ActiveEditorEntityManager->AddTagToEntity(Entity, ColumnType);
-			return true;
+			ExecuteAddColumnCommand(Row, ColumnType);
 		}
-		else if (ColumnType->IsChildOf(FMassFragment::StaticStruct()))
+		else
 		{
-			FStructView Column = ActiveEditorEntityManager->GetFragmentDataStruct(Entity, ColumnType);
-			if (!Column.IsValid())
-			{
-				ActiveEditorEntityManager->AddFragmentToEntity(Entity, ColumnType);
-				return true;
-			}
+			AddPendingCommand(Row, FAddColumnCommand{ .ColumnType = ColumnType });
 		}
+		return true;
 	}
 	return false;
 }
@@ -302,16 +299,15 @@ bool UTypedElementDatabase::AddColumn(TypedElementRowHandle Row, FTopLevelAssetP
 
 void UTypedElementDatabase::RemoveColumn(TypedElementRowHandle Row, const UScriptStruct* ColumnType)
 {
-	FMassEntityHandle Entity = FMassEntityHandle::FromNumber(Row);
-	if (ActiveEditorEntityManager && ActiveEditorEntityManager->IsEntityActive(Entity))
+	if (ActiveEditorEntityManager)
 	{
-		if (ColumnType->IsChildOf(FMassTag::StaticStruct()))
+		if (HasRowBeenAssigned(Row))
 		{
-			ActiveEditorEntityManager->RemoveTagFromEntity(Entity, ColumnType);
+			ExecuteRemoveColumnCommand(Row, ColumnType);
 		}
-		else if (ColumnType->IsChildOf(FMassFragment::StaticStruct()))
+		else
 		{
-			ActiveEditorEntityManager->RemoveFragmentFromEntity(Entity, ColumnType);
+			AddPendingCommand(Row, FRemoveColumnCommand{ .ColumnType = ColumnType });
 		}
 	}
 }
@@ -414,7 +410,7 @@ ColumnDataResult UTypedElementDatabase::GetColumnData(TypedElementRowHandle Row,
 bool UTypedElementDatabase::AddColumns(TypedElementRowHandle Row, TConstArrayView<const UScriptStruct*> Columns)
 {
 	FMassEntityHandle Entity = FMassEntityHandle::FromNumber(Row);
-	if (ActiveEditorEntityManager && ActiveEditorEntityManager->IsEntityActive(Entity))
+	if (ActiveEditorEntityManager)
 	{
 		FMassArchetypeHandle Archetype = ActiveEditorEntityManager->GetArchetypeForEntity(Entity);
 
@@ -422,9 +418,14 @@ bool UTypedElementDatabase::AddColumns(TypedElementRowHandle Row, TConstArrayVie
 		FMassTagBitSet TagsToAdd;
 		if (ColumnsToBitSets(Columns, FragmentsToAdd, TagsToAdd))
 		{
-			FMassArchetypeCompositionDescriptor AddComposition(
-				MoveTemp(FragmentsToAdd), MoveTemp(TagsToAdd), FMassChunkFragmentBitSet(), FMassSharedFragmentBitSet());
-			ActiveEditorEntityManager->AddCompositionToEntity_GetDelta(Entity, AddComposition);
+			if (ActiveEditorEntityManager->IsEntityActive(Entity))
+			{
+				ExecuteAddColumnsCommand(Row, FragmentsToAdd, TagsToAdd);
+			}
+			else
+			{
+				AddPendingCommand(Row, FAddColumnsCommand{ .FragmentsToAdd = MoveTemp(FragmentsToAdd), .TagsToAdd = MoveTemp(TagsToAdd) });
+			}
 			return true;
 		}
 	}
@@ -434,7 +435,7 @@ bool UTypedElementDatabase::AddColumns(TypedElementRowHandle Row, TConstArrayVie
 void UTypedElementDatabase::RemoveColumns(TypedElementRowHandle Row, TConstArrayView<const UScriptStruct*> Columns)
 {
 	FMassEntityHandle Entity = FMassEntityHandle::FromNumber(Row);
-	if (ActiveEditorEntityManager && ActiveEditorEntityManager->IsEntityActive(Entity))
+	if (ActiveEditorEntityManager)
 	{
 		FMassArchetypeHandle Archetype = ActiveEditorEntityManager->GetArchetypeForEntity(Entity);
 
@@ -442,9 +443,15 @@ void UTypedElementDatabase::RemoveColumns(TypedElementRowHandle Row, TConstArray
 		FMassTagBitSet TagsToRemove;
 		if (ColumnsToBitSets(Columns, FragmentsToRemove, TagsToRemove))
 		{
-			FMassArchetypeCompositionDescriptor RemoveComposition(
-				MoveTemp(FragmentsToRemove), MoveTemp(TagsToRemove), FMassChunkFragmentBitSet(), FMassSharedFragmentBitSet());
-			ActiveEditorEntityManager->RemoveCompositionFromEntity(Entity, RemoveComposition);
+			if (ActiveEditorEntityManager->IsEntityActive(Entity))
+			{
+				ExecuteRemoveColumnsCommand(Row, FragmentsToRemove, TagsToRemove);
+			}
+			else
+			{
+				AddPendingCommand(Row, 
+					FRemoveColumnsCommand{ .FragmentsToRemove = MoveTemp(FragmentsToRemove), .TagsToRemove = MoveTemp(TagsToRemove) });
+			}
 		}
 	}
 }
@@ -522,7 +529,7 @@ bool UTypedElementDatabase::BatchAddRemoveColumns(TConstArrayView<TypedElementRo
 				EntityCollections.Emplace(It.Key(), It.Value(), FMassArchetypeEntityCollection::EDuplicatesHandling::FoldDuplicates);
 			}
 
-			// Batch update usint the appropriate fragment/bit sets.
+			// Batch update using the appropriate fragment/bit sets.
 			if (bMustUpdateFragments)
 			{
 				ActiveEditorEntityManager->BatchChangeFragmentCompositionForEntities(EntityCollections, FragmentsToAdd, FragmentsToRemove);
@@ -725,6 +732,97 @@ bool UTypedElementDatabase::ColumnsToBitSets(TConstArrayView<const UScriptStruct
 	return bResult;
 }
 
+template<typename T>
+void UTypedElementDatabase::AddPendingCommand(TypedElementRowHandle Row, T&& Args)
+{
+	FCommand Command;
+	Command.Row = Row;
+	Command.Data.Emplace<T>(Forward<T>(Args));
+	PendingCommands.Add(MoveTemp(Command));
+}
+
+void UTypedElementDatabase::ProcessPendingCommands()
+{
+	PendingCommands.StableSort(
+		[](const FCommand& Lhs, const FCommand& Rhs)
+		{
+			return Lhs.Row < Rhs.Row;
+		});
+	
+	struct FProcessor
+	{
+		UTypedElementDatabase* This;
+		TypedElementRowHandle Row;
+		void operator()(const FAddColumnCommand& Command)		{ This->ExecuteAddColumnCommand(Row, Command.ColumnType.Get()); }
+		void operator()(const FAddColumnsCommand& Command)		{ This->ExecuteAddColumnsCommand(Row, Command.FragmentsToAdd, Command.TagsToAdd); }
+		void operator()(const FRemoveColumnCommand& Command)	{ This->ExecuteRemoveColumnCommand(Row, Command.ColumnType.Get()); }
+		void operator()(const FRemoveColumnsCommand& Command)	{ This->ExecuteRemoveColumnsCommand(Row, Command.FragmentsToRemove, Command.TagsToRemove); }
+	};
+	FProcessor Processor;
+	Processor.This = this;
+
+	for (FCommand& Command : PendingCommands)
+	{
+		if (HasRowBeenAssigned(Command.Row))
+		{
+			Processor.Row = Command.Row;
+			Visit(Processor, Command.Data);
+		}
+	}
+	PendingCommands.Reset();
+}
+
+void UTypedElementDatabase::ExecuteAddColumnCommand(TypedElementRowHandle Row, const UScriptStruct* ColumnType)
+{
+	if (ColumnType)
+	{
+		FMassEntityHandle Entity = FMassEntityHandle::FromNumber(Row);
+		if (ColumnType->IsChildOf(FMassTag::StaticStruct()))
+		{
+			ActiveEditorEntityManager->AddTagToEntity(Entity, ColumnType);
+		}
+		else if (ColumnType->IsChildOf(FMassFragment::StaticStruct()))
+		{
+			FStructView Column = ActiveEditorEntityManager->GetFragmentDataStruct(Entity, ColumnType);
+			// Only add if not already added to avoid asserts from Mass.
+			if (!Column.IsValid())
+			{
+				ActiveEditorEntityManager->AddFragmentToEntity(Entity, ColumnType);
+			}
+		}
+	}
+}
+
+void UTypedElementDatabase::ExecuteRemoveColumnCommand(TypedElementRowHandle Row, const UScriptStruct* ColumnType)
+{
+	if (ColumnType)
+	{
+		FMassEntityHandle Entity = FMassEntityHandle::FromNumber(Row);
+		if (ColumnType->IsChildOf(FMassTag::StaticStruct()))
+		{
+			ActiveEditorEntityManager->RemoveTagFromEntity(Entity, ColumnType);
+		}
+		else if (ColumnType->IsChildOf(FMassFragment::StaticStruct()))
+		{
+			ActiveEditorEntityManager->RemoveFragmentFromEntity(Entity, ColumnType);
+		}
+	}
+}
+
+void UTypedElementDatabase::ExecuteAddColumnsCommand(TypedElementRowHandle Row, FMassFragmentBitSet FragmentsToAdd, FMassTagBitSet TagsToAdd)
+{
+	FMassArchetypeCompositionDescriptor AddComposition(
+		MoveTemp(FragmentsToAdd), MoveTemp(TagsToAdd), FMassChunkFragmentBitSet(), FMassSharedFragmentBitSet());
+	ActiveEditorEntityManager->AddCompositionToEntity_GetDelta(FMassEntityHandle::FromNumber(Row), AddComposition);
+}
+
+void UTypedElementDatabase::ExecuteRemoveColumnsCommand(TypedElementRowHandle Row, FMassFragmentBitSet FragmentsToRemove, FMassTagBitSet TagsToRemove)
+{
+	FMassArchetypeCompositionDescriptor RemoveComposition(
+		MoveTemp(FragmentsToRemove), MoveTemp(TagsToRemove), FMassChunkFragmentBitSet(), FMassSharedFragmentBitSet());
+	ActiveEditorEntityManager->RemoveCompositionFromEntity(FMassEntityHandle::FromNumber(Row), RemoveComposition);
+}
+
 void UTypedElementDatabase::PreparePhase(EQueryTickPhase Phase, float DeltaTime)
 {
 	if (ActiveEditorEntityManager)
diff --git a/Engine/Plugins/Experimental/TypedElementsDataStorage/Source/TypedElementsDataStorage/Private/TypedElementDatabase.h b/Engine/Plugins/Experimental/TypedElementsDataStorage/Source/TypedElementsDataStorage/Private/TypedElementDatabase.h
--- a/Engine/Plugins/Experimental/TypedElementsDataStorage/Source/TypedElementsDataStorage/Private/TypedElementDatabase.h
+++ b/Engine/Plugins/Experimental/TypedElementsDataStorage/Source/TypedElementsDataStorage/Private/TypedElementDatabase.h
@@ -5,6 +5,7 @@
 #include "Templates/SharedPointer.h"
 #include "Elements/Interfaces/TypedElementDataStorageInterface.h"
 #include "MassArchetypeTypes.h"
+#include "Misc/TVariant.h"
 #include "Queries/TypedElementExtendedQueryStore.h"
 #include "TypedElementDatabaseScratchBuffer.h"
 #include "UObject/ObjectMacros.h"
@@ -87,9 +88,43 @@ public:
 	void DebugPrintQueryCallbacks(FOutputDevice& Output);
 
 private:	
+	struct FAddColumnCommand
+	{
+		TWeakObjectPtr<const UScriptStruct> ColumnType;
+	};
+	struct FAddColumnsCommand
+	{
+		FMassFragmentBitSet FragmentsToAdd;
+		FMassTagBitSet TagsToAdd;
+	};
+	struct FRemoveColumnCommand
+	{
+		TWeakObjectPtr<const UScriptStruct> ColumnType;
+	};
+	struct FRemoveColumnsCommand
+	{
+		FMassFragmentBitSet FragmentsToRemove;
+		FMassTagBitSet TagsToRemove;
+	};
+	using CommandData = TVariant<FAddColumnCommand, FAddColumnsCommand, FRemoveColumnCommand, FRemoveColumnsCommand>;
+
+	struct FCommand
+	{
+		TypedElementRowHandle Row;
+		CommandData Data;
+	};
+	
 	/** Converts a set of column types into Mass specific fragment and tag bit sets. Returns true if any values were added. */
 	static bool ColumnsToBitSets(TConstArrayView<const UScriptStruct*> Columns, FMassFragmentBitSet& Fragments, FMassTagBitSet& Tags);
 
+	template<typename T>
+	void AddPendingCommand(TypedElementRowHandle Row, T&& Args);
+	void ProcessPendingCommands();
+	void ExecuteAddColumnCommand(TypedElementRowHandle Row, const UScriptStruct* ColumnType);
+	void ExecuteAddColumnsCommand(TypedElementRowHandle Row, FMassFragmentBitSet FragmentsToAdd, FMassTagBitSet TagsToAdd);
+	void ExecuteRemoveColumnCommand(TypedElementRowHandle Row, const UScriptStruct* ColumnType);
+	void ExecuteRemoveColumnsCommand(TypedElementRowHandle Row, FMassFragmentBitSet FragmentsToRemove, FMassTagBitSet TagsToRemove);
+	
 	void PreparePhase(EQueryTickPhase Phase, float DeltaTime);
 	void FinalizePhase(EQueryTickPhase Phase, float DeltaTime);
 	void Reset();
@@ -106,4 +141,6 @@ private:
 
 	TSharedPtr<FMassEntityManager> ActiveEditorEntityManager;
 	TSharedPtr<FMassProcessingPhaseManager> ActiveEditorPhaseManager;
+
+	TArray<FCommand> PendingCommands;
 };
\ No newline at end of file
