Step towards having TEDS treat Actors the same way as UObjects.

This change registers Actors in the same way as UObjects are registered. For the time being this is still using FMassFragment and is behind a define so additional testing can be done and reverted quickly if needed.
This change also addresses a crash when trying to add objects that were destroyed before TEDS has had a chance to process them.

diff --git a/Engine/Plugins/Experimental/TypedElementsDataStorage/Source/TypedElementsDataStorage/Private/TypedElementDatabaseCompatibility.cpp b/Engine/Plugins/Experimental/TypedElementsDataStorage/Source/TypedElementsDataStorage/Private/TypedElementDatabaseCompatibility.cpp
--- a/Engine/Plugins/Experimental/TypedElementsDataStorage/Source/TypedElementsDataStorage/Private/TypedElementDatabaseCompatibility.cpp
+++ b/Engine/Plugins/Experimental/TypedElementsDataStorage/Source/TypedElementsDataStorage/Private/TypedElementDatabaseCompatibility.cpp
@@ -10,10 +10,13 @@
 #include "Elements/Columns/TypedElementPackageColumns.h"
 #include "Elements/Columns/TypedElementTransformColumns.h"
 #include "Elements/Columns/TypedElementTypeInfoColumns.h"
-#include "MassActorEditorSubsystem.h"
 #include "MassActorSubsystem.h"
 #include "TypedElementDataStorageProfilingMacros.h"
 
+#if TEDS_SEPARATE_ACTOR_REGISTRATION
+#	include "MassActorEditorSubsystem.h"
+#endif
+
 void UTypedElementDatabaseCompatibility::Initialize(ITypedElementDataStorageInterface* StorageInterface)
 {
 	checkf(StorageInterface, TEXT("Typed Element's Database compatibility manager is being initialized with an invalid storage target."));
@@ -82,11 +85,13 @@ TypedElementRowHandle UTypedElementDatabaseCompatibility::AddCompatibleObjectExp
 
 TypedElementRowHandle UTypedElementDatabaseCompatibility::AddCompatibleObjectExplicit(UObject* Object, TypedElementTableHandle Table)
 {
+#if TEDS_SEPARATE_ACTOR_REGISTRATION
 	if (Object->IsA<AActor>())
 	{
 		return AddCompatibleObjectExplicit(static_cast<AActor*>(Object));
 	}
 	else
+#endif
 	{
 		if (ensureMsgf(Storage, TEXT("Trying to add a UObject to Typed Element's Data Storage before the storage is available.")) && 
 			ShouldAddObject(Object))
@@ -113,6 +118,7 @@ TypedElementRowHandle UTypedElementDatabaseCompatibility::AddCompatibleObjectExp
 
 TypedElementRowHandle UTypedElementDatabaseCompatibility::AddCompatibleObjectExplicit(AActor* Actor, TypedElementTableHandle Table)
 {
+#if TEDS_SEPARATE_ACTOR_REGISTRATION
 	if (ensureMsgf(Storage, TEXT("Trying to add an actor to Typed Element's Data Storage before the storage is available.")) &&
 		ShouldAddObject(Actor))
 	{
@@ -126,6 +132,9 @@ TypedElementRowHandle UTypedElementDatabaseCompatibility::AddCompatibleObjectExp
 		return ReservedRow;
 	}
 	return TypedElementInvalidRowHandle;
+#else
+	return AddCompatibleObjectExplicit(static_cast<UObject*>(Actor), Table);
+#endif
 }
 
 TypedElementRowHandle UTypedElementDatabaseCompatibility::AddCompatibleObjectExplicit(void* Object, TWeakObjectPtr<const UScriptStruct> TypeInfo)
@@ -152,11 +161,13 @@ TypedElementRowHandle UTypedElementDatabaseCompatibility::AddCompatibleObjectExp
 
 void UTypedElementDatabaseCompatibility::RemoveCompatibleObjectExplicit(UObject* Object)
 {
+#if TEDS_SEPARATE_ACTOR_REGISTRATION
 	if (Object->IsA<AActor>())
 	{
 		RemoveCompatibleObjectExplicit(static_cast<AActor*>(Object));
 	}
 	else
+#else
 	{
 		checkf(Storage, TEXT("Removing compatible objects is not supported before Typed Element's Database compatibility manager has been initialized."));
 		TypedElementRowHandle Row;
@@ -170,6 +181,7 @@ void UTypedElementDatabaseCompatibility::RemoveCompatibleObjectExplicit(UObject*
 			Storage->RemoveRow(Row);
 		}
 	}
+#endif
 }
 
 void UTypedElementDatabaseCompatibility::RemoveCompatibleObjectExplicit(void* Object)
@@ -189,6 +201,7 @@ void UTypedElementDatabaseCompatibility::RemoveCompatibleObjectExplicit(void* Ob
 
 void UTypedElementDatabaseCompatibility::RemoveCompatibleObjectExplicit(AActor* Actor)
 {
+#if TEDS_SEPARATE_ACTOR_REGISTRATION
 	checkf(Storage, TEXT("Removing compatible objects is not supported before Typed Element's Database compatibility manager has been initialized."));
 
 	// If there is no actor subsystem it means that the world has been destroyed, including the MASS instance,
@@ -218,12 +231,16 @@ void UTypedElementDatabaseCompatibility::RemoveCompatibleObjectExplicit(AActor*
 			}
 		}
 	}
+#else
+	RemoveCompatibleObjectExplicit(static_cast<UObject*>(Actor));
+#endif
 }
 
 TypedElementRowHandle UTypedElementDatabaseCompatibility::FindRowWithCompatibleObjectExplicit(const UObject* Object) const
 {
 	if (Object && Storage && Storage->IsAvailable())
 	{
+#if TEDS_SEPARATE_ACTOR_REGISTRATION
 		const AActor* Actor = Cast<AActor>(Object);
 		if (Actor && ActorSubsystem)
 		{
@@ -231,6 +248,7 @@ TypedElementRowHandle UTypedElementDatabaseCompatibility::FindRowWithCompatibleO
 			return Entity.IsValid() ? Entity.AsNumber() : DealiasObject(Object);
 		}
 		else
+#endif
 		{
 			const TypedElementRowHandle* Row = ReverseObjectLookup.Find(Object);
 			return Row ? *Row : DealiasObject(Object);
@@ -243,6 +261,7 @@ TypedElementRowHandle UTypedElementDatabaseCompatibility::FindRowWithCompatibleO
 
 TypedElementRowHandle UTypedElementDatabaseCompatibility::FindRowWithCompatibleObjectExplicit(const AActor* Actor) const
 {
+#if TEDS_SEPARATE_ACTOR_REGISTRATION
 	if (Storage && ActorSubsystem && Storage->IsAvailable())
 	{
 		FMassEntityHandle Entity = ActorSubsystem->GetEntityHandleFromActor(Actor);
@@ -252,6 +271,9 @@ TypedElementRowHandle UTypedElementDatabaseCompatibility::FindRowWithCompatibleO
 	{
 		return TypedElementInvalidRowHandle;
 	}
+#else
+	return FindRowWithCompatibleObjectExplicit(static_cast<const UObject*>(Actor));
+#endif
 }
 
 TypedElementRowHandle UTypedElementDatabaseCompatibility::FindRowWithCompatibleObjectExplicit(const void* Object) const
@@ -266,16 +288,20 @@ TypedElementRowHandle UTypedElementDatabaseCompatibility::FindRowWithCompatibleO
 
 void UTypedElementDatabaseCompatibility::Prepare()
 {
+#if TEDS_SEPARATE_ACTOR_REGISTRATION
 	UMassActorEditorSubsystem* MassActorEditorSubsystem = Storage->GetExternalSystem<UMassActorEditorSubsystem>();
 	check(MassActorEditorSubsystem);
 	ActorSubsystem = MassActorEditorSubsystem->GetMutableActorManager().AsShared();
+#endif
 
 	CreateStandardArchetypes();
 }
 
 void UTypedElementDatabaseCompatibility::Reset()
 {
+#if TEDS_SEPARATE_ACTOR_REGISTRATION
 	ActorSubsystem = nullptr;
+#endif
 }
 
 void UTypedElementDatabaseCompatibility::CreateStandardArchetypes()
@@ -335,7 +361,9 @@ void UTypedElementDatabaseCompatibility::Tick()
 	// Delay processing until the required systems are available by not clearing any lists or doing any work.
 	if (Storage && Storage->IsAvailable() && EditorWorld)
 	{
+#if TEDS_SEPARATE_ACTOR_REGISTRATION
 		TickPendingActorRegistration(EditorWorld);
+#endif
 		TickPendingUObjectRegistration();
 		TickPendingExternalObjectRegistration();
 		TickObjectSync();
@@ -397,7 +425,7 @@ void UTypedElementDatabaseCompatibility::PendingRegistration<AddressType>::Remov
 		}
 		else
 		{
-			// Don't shrink the registration array as the array will be reused with a variety of different actor counts.
+			// Don't shrink the registration array as the array will be reused with a variety of different object counts.
 			// If memory size becomes an issue it's better to resize the array once after this loop rather than within the
 			// loop to avoid many resizes happening.
 			constexpr bool bAllowToShrink = false;
@@ -434,6 +462,7 @@ void UTypedElementDatabaseCompatibility::PendingRegistration<AddressType>::Reset
 	ReservedRowHandles.Reset();
 }
 
+#if TEDS_SEPARATE_ACTOR_REGISTRATION
 void UTypedElementDatabaseCompatibility::TickPendingActorRegistration(UWorld* EditorWorld)
 {
 	if (!ActorsPendingRegistration.IsEmpty())
@@ -479,6 +508,7 @@ void UTypedElementDatabaseCompatibility::TickPendingActorRegistration(UWorld* Ed
 		ActorsPendingRegistration.Reset();
 	}
 }
+#endif
 
 void UTypedElementDatabaseCompatibility::TickPendingUObjectRegistration()
 {
@@ -491,16 +521,23 @@ void UTypedElementDatabaseCompatibility::TickPendingUObjectRegistration()
 				[this](const TWeakObjectPtr<UObject>& Object)
 				{
 					UObject* Instance = Object.Get();
-					return Instance && FindRowWithCompatibleObjectExplicit(Instance) != TypedElementInvalidRowHandle;
+					return Instance && Storage->IsRowAvailable(FindRowWithCompatibleObjectExplicit(Instance));
 				});
 		}
 
 		// Add the remaining object to the data storage.
 		for (auto It = UObjectsPendingRegistration.CreateIterator(); It; ++It)
 		{
-			It->Value.ProcessEntries(*Storage, It->Key, 
+			It->Value.ProcessEntries(*Storage, It->Key,
 				[this](TypedElementRowHandle Row, const TWeakObjectPtr<UObject>& Object)
 				{
+#if !TEDS_SEPARATE_ACTOR_REGISTRATION
+					if (AActor* Actor = Cast<AActor>(Object))
+					{
+						constexpr bool bIsOwnedByMass = false;
+						Storage->AddOrGetColumn<FMassActorFragment>(Row)->SetNoHandleMapUpdate(FMassEntityHandle::FromNumber(Row), Actor, bIsOwnedByMass);
+					}
+#endif
 					Storage->AddOrGetColumn<FTypedElementUObjectColumn>(Row, FTypedElementUObjectColumn{ .Object = Object });
 					Storage->AddOrGetColumn<FTypedElementClassTypeInfoColumn>(Row, FTypedElementClassTypeInfoColumn{ .TypeInfo = Object->GetClass() });
 					// Make sure the new row is tagged for update.
@@ -524,7 +561,7 @@ void UTypedElementDatabaseCompatibility::TickPendingExternalObjectRegistration()
 			It->Value.RemoveInvalidEntries(*Storage,
 				[this](const ExternalObjectRegistration& Object)
 				{
-					return Object.Object && FindRowWithCompatibleObjectExplicit(Object.Object) != TypedElementInvalidRowHandle;
+					return Object.Object && Storage->IsRowAvailable(FindRowWithCompatibleObjectExplicit(Object.Object));
 				});
 		}
 
@@ -578,7 +615,7 @@ void UTypedElementDatabaseCompatibility::TickObjectSync()
 
 		{
 			TEDS_EVENT_SCOPE(TEXT("Add SyncFromWorld Tag"));
-			// Tag the rows containing actor data that they need to be synced
+			// Tag the rows containing object data that so they get synced
 			// Note: Watch out for the performance of this, may end up doing a lot of row moves
 			for (TypedElementRowHandle Row : RowHandles)
 			{
diff --git a/Engine/Plugins/Experimental/TypedElementsDataStorage/Source/TypedElementsDataStorage/Private/TypedElementDatabaseCompatibility.h b/Engine/Plugins/Experimental/TypedElementsDataStorage/Source/TypedElementsDataStorage/Private/TypedElementDatabaseCompatibility.h
--- a/Engine/Plugins/Experimental/TypedElementsDataStorage/Source/TypedElementsDataStorage/Private/TypedElementDatabaseCompatibility.h
+++ b/Engine/Plugins/Experimental/TypedElementsDataStorage/Source/TypedElementsDataStorage/Private/TypedElementDatabaseCompatibility.h
@@ -11,6 +11,10 @@
 
 #include "TypedElementDatabaseCompatibility.generated.h"
 
+#if !defined(TEDS_SEPARATE_ACTOR_REGISTRATION)
+#	define TEDS_SEPARATE_ACTOR_REGISTRATION 0
+#endif
+
 class AActor;
 class ITypedElementDataStorageInterface;
 struct FMassActorManager;
@@ -64,7 +68,9 @@ private:
 	TypedElementRowHandle DealiasObject(const UObject* Object) const;
 
 	void Tick();
+#if TEDS_SEPARATE_ACTOR_REGISTRATION
 	void TickPendingActorRegistration(UWorld* EditorWorld);
+#endif
 	void TickPendingUObjectRegistration();
 	void TickPendingExternalObjectRegistration();
 	void TickObjectSync();
@@ -98,7 +104,9 @@ private:
 		void* Object;
 		TWeakObjectPtr<const UScriptStruct> TypeInfo;
 	};
+#if TEDS_SEPARATE_ACTOR_REGISTRATION
 	TMap<TypedElementTableHandle, PendingRegistration<TWeakObjectPtr<AActor>>> ActorsPendingRegistration;
+#endif
 	TMap<TypedElementTableHandle, PendingRegistration<TWeakObjectPtr<UObject>>> UObjectsPendingRegistration;
 	TMap<TypedElementTableHandle, PendingRegistration<ExternalObjectRegistration>> ExternalObjectsPendingRegistration;
 	
@@ -112,7 +120,9 @@ private:
 	TypedElementTableHandle StandardUObjectTable{ TypedElementInvalidTableHandle };
 	TypedElementTableHandle StandardExternalObjectTable{ TypedElementInvalidTableHandle };
 	ITypedElementDataStorageInterface* Storage{ nullptr };
+#if TEDS_SEPARATE_ACTOR_REGISTRATION
 	TSharedPtr<FMassActorManager> ActorSubsystem;
+#endif
 
 	TMap<void*, TypedElementRowHandle> ReverseObjectLookup;
 
