Moved TEDS' counter widgets behind a console command.

diff --git a/Engine/Plugins/Experimental/TypedElementsDataStorage/Source/TypedElementsDataStorageUI/Private/TypedElementCounterWidget.cpp b/Engine/Plugins/Experimental/TypedElementsDataStorage/Source/TypedElementsDataStorageUI/Private/TypedElementCounterWidget.cpp
--- a/Engine/Plugins/Experimental/TypedElementsDataStorage/Source/TypedElementsDataStorageUI/Private/TypedElementCounterWidget.cpp
+++ b/Engine/Plugins/Experimental/TypedElementsDataStorage/Source/TypedElementsDataStorageUI/Private/TypedElementCounterWidget.cpp
@@ -19,15 +19,28 @@
 #define LOCTEXT_NAMESPACE "TypedElementUI_CounterWidget"
 
 
+FAutoConsoleCommand EnableCounterWidgetsConsoleCommand(
+	TEXT("TEDS.UI.EnableCounterWidgets"),
+	TEXT("Adds registered counter widgets to the bottom right status bar of the main editor window."),
+	FConsoleCommandDelegate::CreateLambda([]()
+		{
+			UTypedElementCounterWidgetFactory::EnableCounterWidgets();
+		}));
+
 //
 // UTypedElementCounterWidgetFactory
 //
 
 FName UTypedElementCounterWidgetFactory::WigetPurpose(TEXT("LevelEditor.StatusBar.ToolBar"));
+bool UTypedElementCounterWidgetFactory::bAreCounterWidgetsEnabled{ false };
+bool UTypedElementCounterWidgetFactory::bHasBeenSetup{ false };
 
 UTypedElementCounterWidgetFactory::UTypedElementCounterWidgetFactory()
 {
-	IMainFrameModule::Get().OnMainFrameCreationFinished().AddStatic(&UTypedElementCounterWidgetFactory::SetupMainWindowIntegrations);
+	if (bAreCounterWidgetsEnabled)
+	{
+		IMainFrameModule::Get().OnMainFrameCreationFinished().AddStatic(&UTypedElementCounterWidgetFactory::SetupMainWindowIntegrations);
+	}
 }
 
 void UTypedElementCounterWidgetFactory::RegisterQueries(ITypedElementDataStorageInterface& DataStorage) const
@@ -101,39 +114,49 @@ void UTypedElementCounterWidgetFactory::RegisterWidgetConstructors(ITypedElement
 	DataStorageUi.RegisterWidgetFactory(WigetPurpose, MoveTemp(WidgetCounter));
 }
 
-void UTypedElementCounterWidgetFactory::SetupMainWindowIntegrations(TSharedPtr<SWindow> ParentWindow, bool bIsRunningStartupDialog)
+void UTypedElementCounterWidgetFactory::EnableCounterWidgets()
 {
-	UTypedElementRegistry* Registry = UTypedElementRegistry::GetInstance();
-	checkf(Registry, TEXT(
-		"FTypedElementsDataStorageUiModule didn't find the UTypedElementRegistry during main window integration when it should be available."));
+	bAreCounterWidgetsEnabled = true;
+	UTypedElementCounterWidgetFactory::SetupMainWindowIntegrations(nullptr, false);
+}
 
-	ITypedElementDataStorageUiInterface* UiInterface = Registry->GetMutableDataStorageUi();
-	checkf(UiInterface, TEXT(
-		"FTypedElementsDataStorageUiModule tried to integrate with the main window before the "
-		"Typed Elements Data Storage UI interface is available."));
+void UTypedElementCounterWidgetFactory::SetupMainWindowIntegrations(TSharedPtr<SWindow> ParentWindow, bool bIsRunningStartupDialog)
+{
+	if (!bHasBeenSetup)
+	{
+		UTypedElementRegistry* Registry = UTypedElementRegistry::GetInstance();
+		checkf(Registry, TEXT(
+			"FTypedElementsDataStorageUiModule didn't find the UTypedElementRegistry during main window integration when it should be available."));
 
-	UToolMenu* Menu = UToolMenus::Get()->ExtendMenu(WigetPurpose);
+		ITypedElementDataStorageUiInterface* UiInterface = Registry->GetMutableDataStorageUi();
+		checkf(UiInterface, TEXT(
+			"FTypedElementsDataStorageUiModule tried to integrate with the main window before the "
+			"Typed Elements Data Storage UI interface is available."));
 
-	TArray<TSharedRef<SWidget>> Widgets;
-	UiInterface->ConstructWidgets(WigetPurpose, {},
-		[&Widgets](const TSharedRef<SWidget>& NewWidget, TypedElementRowHandle Row)
-		{
-			Widgets.Add(NewWidget);
-		});
+		UToolMenu* Menu = UToolMenus::Get()->ExtendMenu(WigetPurpose);
 
-	if (!Widgets.IsEmpty())
-	{
-		FToolMenuSection& Section = Menu->AddSection("DataStorageSection");
-		int32 WidgetCount = Widgets.Num();
+		TArray<TSharedRef<SWidget>> Widgets;
+		UiInterface->ConstructWidgets(WigetPurpose, {},
+			[&Widgets](const TSharedRef<SWidget>& NewWidget, TypedElementRowHandle Row)
+			{
+				Widgets.Add(NewWidget);
+			});
 
-		Section.AddEntry(FToolMenuEntry::InitWidget("DataStorageStatusBarWidget_0", MoveTemp(Widgets[0]), FText::GetEmpty()));
-		for (int32 I = 1; I < WidgetCount; ++I)
+		if (!Widgets.IsEmpty())
 		{
-			Section.AddSeparator(FName(*FString::Format(TEXT("DataStorageStatusBarWidgetDivider_{0}"), { FString::FromInt(I) })));
-			Section.AddEntry(FToolMenuEntry::InitWidget(
-				FName(*FString::Format(TEXT("DataStorageStatusBarWidget_{0}"), { FString::FromInt(I) })),
-				MoveTemp(Widgets[I]), FText::GetEmpty()));
+			FToolMenuSection& Section = Menu->AddSection("DataStorageSection");
+			int32 WidgetCount = Widgets.Num();
+
+			Section.AddEntry(FToolMenuEntry::InitWidget("DataStorageStatusBarWidget_0", MoveTemp(Widgets[0]), FText::GetEmpty()));
+			for (int32 I = 1; I < WidgetCount; ++I)
+			{
+				Section.AddSeparator(FName(*FString::Format(TEXT("DataStorageStatusBarWidgetDivider_{0}"), { FString::FromInt(I) })));
+				Section.AddEntry(FToolMenuEntry::InitWidget(
+					FName(*FString::Format(TEXT("DataStorageStatusBarWidget_{0}"), { FString::FromInt(I) })),
+					MoveTemp(Widgets[I]), FText::GetEmpty()));
+			}
 		}
+		bHasBeenSetup = true;
 	}
 }
 
diff --git a/Engine/Plugins/Experimental/TypedElementsDataStorage/Source/TypedElementsDataStorageUI/Public/TypedElementCounterWidget.h b/Engine/Plugins/Experimental/TypedElementsDataStorage/Source/TypedElementsDataStorageUI/Public/TypedElementCounterWidget.h
--- a/Engine/Plugins/Experimental/TypedElementsDataStorage/Source/TypedElementsDataStorageUI/Public/TypedElementCounterWidget.h
+++ b/Engine/Plugins/Experimental/TypedElementsDataStorage/Source/TypedElementsDataStorageUI/Public/TypedElementCounterWidget.h
@@ -28,8 +28,13 @@ public:
 	void RegisterWidgetConstructors(ITypedElementDataStorageInterface& DataStorage,
 		ITypedElementDataStorageUiInterface& DataStorageUi) const override;
 
+	static void EnableCounterWidgets();
+
 private:
 	static void SetupMainWindowIntegrations(TSharedPtr<SWindow> ParentWindow, bool bIsRunningStartupDialog);
+
+	static bool bAreCounterWidgetsEnabled;
+	static bool bHasBeenSetup;
 };
 
 /**
