TEDS - Update the TypeInfo widget to have a clickable link to blueprints

diff --git a/Engine/Plugins/Experimental/TypedElementsDataStorage/Source/TypedElementsDataStorageUI/Private/TypedElementTypeInfoWidget.cpp b/Engine/Plugins/Experimental/TypedElementsDataStorage/Source/TypedElementsDataStorageUI/Private/TypedElementTypeInfoWidget.cpp
--- a/Engine/Plugins/Experimental/TypedElementsDataStorage/Source/TypedElementsDataStorageUI/Private/TypedElementTypeInfoWidget.cpp
+++ b/Engine/Plugins/Experimental/TypedElementsDataStorage/Source/TypedElementsDataStorageUI/Private/TypedElementTypeInfoWidget.cpp
@@ -3,7 +3,10 @@
 #include "TypedElementTypeInfoWidget.h"
 
 #include "MassActorSubsystem.h"
+#include "SceneOutlinerHelpers.h"
+#include "Elements/Columns/TypedElementCompatibilityColumns.h"
 #include "Elements/Columns/TypedElementMiscColumns.h"
+#include "Elements/Columns/TypedElementSlateWidgetColumns.h"
 #include "Elements/Columns/TypedElementTypeInfoColumns.h"
 #include "Styling/SlateIconFinder.h"
 #include "Widgets/Text/STextBlock.h"
@@ -44,20 +47,16 @@ TSharedPtr<SWidget> FTypedElementTypeInfoWidgetConstructor::CreateWidget(
 		bUseIcon = *MetaDataEntryView.TryGetExact<bool>();
 	}
 
-	// Can't do this in one line because slate doesn't like multiple SNew in one line
-	TSharedPtr<SWidget> TypeWidget;
 	if(bUseIcon)
 	{
-		TypeWidget = SNew(SImage)
+		return SNew(SImage)
 					.DesiredSizeOverride(FVector2D(16.f, 16.f))
 					.ColorAndOpacity(FSlateColor::UseForeground());
 	}
 	else
 	{
-		TypeWidget = SNew(STextBlock);
+		return SNew(SHorizontalBox);
 	}
-	
-	return TypeWidget;
 }
 
 bool FTypedElementTypeInfoWidgetConstructor::FinalizeWidget(ITypedElementDataStorageInterface* DataStorage,
@@ -84,14 +83,37 @@ bool FTypedElementTypeInfoWidgetConstructor::FinalizeWidget(ITypedElementDataSto
 		}
 		else
 		{
-			checkf(Widget->GetType() == STextBlock::StaticWidgetClass().GetWidgetType(),
+			checkf(Widget->GetType() == SHorizontalBox::StaticWidgetClass().GetWidgetType(),
 				TEXT("Stored widget with FTypedElementLabelWidgetConstructor doesn't match type %s, but was a %s."),
-				*(STextBlock::StaticWidgetClass().GetWidgetType().ToString()),
+				*(SHorizontalBox::StaticWidgetClass().GetWidgetType().ToString()),
 				*(Widget->GetTypeAsString()));
 				
-			STextBlock* WidgetInstance = static_cast<STextBlock*>(Widget.Get());
-
-			WidgetInstance->SetText(FText::FromString(TypeInfoColumn->TypeInfo.Get()->GetName()));
+			SHorizontalBox* WidgetInstance = static_cast<SHorizontalBox*>(Widget.Get());
+
+			TSharedPtr<SWidget> ActualWidget;
+
+			// Check if we have a hyperlink for this object
+			if (const FTypedElementUObjectColumn* ObjectColumn = DataStorage->GetColumn<FTypedElementUObjectColumn>(TargetRow))
+			{
+				ActualWidget = SceneOutliner::FSceneOutlinerHelpers::GetClassHyperlink(ObjectColumn->Object.Get());
+			}
+
+			// If not, we simply show a text block with the type
+			if(!ActualWidget)
+			{
+				ActualWidget = SNew(STextBlock)
+								.ColorAndOpacity(FSlateColor::UseSubduedForeground())
+								.Text(FText::FromString(TypeInfoColumn->TypeInfo.Get()->GetName()));
+			}
+
+
+			WidgetInstance->AddSlot()
+						.AutoWidth()
+						.VAlign(VAlign_Center)
+						.Padding(8, 0, 0, 0)
+						[
+							ActualWidget.ToSharedRef()
+						];
 		}
 	}
 
@@ -100,7 +122,6 @@ bool FTypedElementTypeInfoWidgetConstructor::FinalizeWidget(ITypedElementDataSto
 
 const FSlateBrush* FTypedElementTypeInfoWidgetConstructor::GetIconForRow(ITypedElementDataStorageInterface* DataStorage, TypedElementRowHandle Row, const FTypedElementClassTypeInfoColumn* TypeInfoColumn)
 {
-
 	/* The logic here is very similar to SActorTreeLabel::GetIcon in ActorTreeItem.cpp which allows the actor to specify
 	 * an override for the icon, and has a fallback to the class icon if not.
 	 */
@@ -135,4 +156,4 @@ const FSlateBrush* FTypedElementTypeInfoWidgetConstructor::GetIconForRow(ITypedE
 
 	// Fallback to the regular actor icon if we haven't found any specific icon
 	return FSlateIconFinder::FindIconForClass(AActor::StaticClass()).GetOptionalIcon();;
-}
\ No newline at end of file
+}
diff --git a/Engine/Source/Editor/SceneOutliner/Private/SceneOutlinerActorInfoColumn.cpp b/Engine/Source/Editor/SceneOutliner/Private/SceneOutlinerActorInfoColumn.cpp
--- a/Engine/Source/Editor/SceneOutliner/Private/SceneOutlinerActorInfoColumn.cpp
+++ b/Engine/Source/Editor/SceneOutliner/Private/SceneOutlinerActorInfoColumn.cpp
@@ -19,6 +19,7 @@
 #include "ComponentTreeItem.h"
 #include "FolderTreeItem.h"
 #include "LevelTreeItem.h"
+#include "SceneOutlinerHelpers.h"
 #include "WorldTreeItem.h"
 #include "WorldPartition/DataLayer/DataLayerInstance.h"
 #include "WorldPartition/DataLayer/WorldDataLayers.h"
@@ -157,34 +158,7 @@ TSharedPtr<SWidget> FTypeInfoColumn::ConstructClassHyperlink( ISceneOutlinerTree
 	{
 		if (AActor* Actor = ActorItem->Actor.Get())
 		{
-			if (UClass* ActorClass = Actor->GetClass())
-			{
-				// Always show blueprints
-				const bool bIsBlueprintClass = UBlueprint::GetBlueprintFromClass(ActorClass) != nullptr;
-
-				// Also show game or game plugin native classes (but not engine classes as that makes the scene outliner pretty noisy)
-				bool bIsGameClass = false;
-				if (!bIsBlueprintClass)
-				{
-					UPackage* Package = ActorClass->GetOutermost();
-					const FString ModuleName = FPackageName::GetShortName(Package->GetFName());
-
-					FModuleStatus PackageModuleStatus;
-					if (FModuleManager::Get().QueryModule(*ModuleName, /*out*/ PackageModuleStatus))
-					{
-						bIsGameClass = PackageModuleStatus.bIsGameModule;
-					}
-				}
-
-				if (bIsBlueprintClass || bIsGameClass)
-				{
-					FEditorClassUtils::FSourceLinkParams SourceLinkParams;
-					SourceLinkParams.Object = Actor;
-					SourceLinkParams.bUseDefaultFormat = true;
-
-					return FEditorClassUtils::GetSourceLink(ActorClass, SourceLinkParams);
-				}
-			}
+			return SceneOutliner::FSceneOutlinerHelpers::GetClassHyperlink(Actor);
 		}
 	}
 
diff --git a/Engine/Source/Editor/SceneOutliner/Private/SceneOutlinerHelpers.cpp b/Engine/Source/Editor/SceneOutliner/Private/SceneOutlinerHelpers.cpp
--- a/Engine/Source/Editor/SceneOutliner/Private/SceneOutlinerHelpers.cpp
+++ b/Engine/Source/Editor/SceneOutliner/Private/SceneOutlinerHelpers.cpp
@@ -4,6 +4,10 @@
 #include "ActorTreeItem.h"
 #include "ActorDescTreeItem.h"
 #include "ActorFolderTreeItem.h"
+#include "EditorClassUtils.h"
+#include "Engine/Blueprint.h"
+#include "Misc/PackageName.h"
+#include "Modules/ModuleManager.h"
 #include "WorldPartition/WorldPartitionActorDesc.h"
 #include "UObject/Package.h"
 
@@ -77,4 +81,42 @@ namespace SceneOutliner
 
 		return nullptr;
 	}
+
+TSharedPtr<SWidget> FSceneOutlinerHelpers::GetClassHyperlink(UObject* InObject)
+	{
+		if (InObject)
+		{
+			if (UClass* Class = InObject->GetClass())
+			{
+				// Always show blueprints
+				const bool bIsBlueprintClass = UBlueprint::GetBlueprintFromClass(Class) != nullptr;
+
+				// Also show game or game plugin native classes (but not engine classes as that makes the scene outliner pretty noisy)
+				bool bIsGameClass = false;
+				if (!bIsBlueprintClass)
+				{
+					UPackage* Package = Class->GetOutermost();
+					const FString ModuleName = FPackageName::GetShortName(Package->GetFName());
+
+					FModuleStatus PackageModuleStatus;
+					if (FModuleManager::Get().QueryModule(*ModuleName, /*out*/ PackageModuleStatus))
+					{
+						bIsGameClass = PackageModuleStatus.bIsGameModule;
+					}
+				}
+
+				if (bIsBlueprintClass || bIsGameClass)
+				{
+					FEditorClassUtils::FSourceLinkParams SourceLinkParams;
+					SourceLinkParams.Object = InObject;
+					SourceLinkParams.bUseDefaultFormat = true;
+
+					return FEditorClassUtils::GetSourceLink(Class, SourceLinkParams);
+				}
+			}
+		}
+
+		return nullptr;
+	}
+
 ;}
\ No newline at end of file
diff --git a/Engine/Source/Editor/SceneOutliner/Public/SceneOutlinerHelpers.h b/Engine/Source/Editor/SceneOutliner/Public/SceneOutlinerHelpers.h
--- a/Engine/Source/Editor/SceneOutliner/Public/SceneOutlinerHelpers.h
+++ b/Engine/Source/Editor/SceneOutliner/Public/SceneOutlinerHelpers.h
@@ -12,5 +12,6 @@ namespace SceneOutliner
 	public:
 		static FString GetExternalPackageName(const ISceneOutlinerTreeItem& TreeItem);
 		static UPackage* GetExternalPackage(const ISceneOutlinerTreeItem& TreeItem);
+		static TSharedPtr<SWidget> GetClassHyperlink(UObject* InObject);
 	};
 };
diff --git a/Engine/Source/Runtime/TypedElementFramework/Public/Elements/Columns/TypedElementTypeInfoColumns.h b/Engine/Source/Runtime/TypedElementFramework/Public/Elements/Columns/TypedElementTypeInfoColumns.h
--- a/Engine/Source/Runtime/TypedElementFramework/Public/Elements/Columns/TypedElementTypeInfoColumns.h
+++ b/Engine/Source/Runtime/TypedElementFramework/Public/Elements/Columns/TypedElementTypeInfoColumns.h
@@ -14,7 +14,7 @@ class UScriptStruct;
 /**
  * Column that stores type information for classes.
  */
-USTRUCT(meta = (DisplayName = "Class type info"))
+USTRUCT(meta = (DisplayName = "Type"))
 struct FTypedElementClassTypeInfoColumn final : public FTypedElementDataStorageColumn
 {
 	GENERATED_BODY()
