Fixed the same object getting (un)registered with TEDS Compat causing a crash.

In some cases the UI can rapidly cause objects to be created  and destroyed, for instance by dragging sliders around. It's possible for this to happen within the same frame, causing the same object to be unregistered and then registered again multiple times. This in turn causes the same object to be pending registration multiple times. The first registration will succeed but not the following ones as there's already an object registered. This change removes the unregistered version from the queue before processing.

diff --git a/Engine/Plugins/Experimental/TypedElementsDataStorage/Source/TypedElementsDataStorage/Private/TypedElementDatabaseCompatibility.cpp b/Engine/Plugins/Experimental/TypedElementsDataStorage/Source/TypedElementsDataStorage/Private/TypedElementDatabaseCompatibility.cpp
--- a/Engine/Plugins/Experimental/TypedElementsDataStorage/Source/TypedElementsDataStorage/Private/TypedElementDatabaseCompatibility.cpp
+++ b/Engine/Plugins/Experimental/TypedElementsDataStorage/Source/TypedElementsDataStorage/Private/TypedElementDatabaseCompatibility.cpp
@@ -418,7 +418,7 @@ void UTypedElementDatabaseCompatibility::PendingRegistration<AddressType>::Remov
 	TypedElementRowHandle* RowHandleIt = ReservedRowHandles.GetData();
 	while (AddressIt != AddressEnd)
 	{
-		if (Validator(*AddressIt))
+		if (StorageInterface.IsRowAvailable(*RowHandleIt) && Validator(*AddressIt))
 		{
 			++AddressIt;
 			++RowHandleIt;
@@ -430,8 +430,8 @@ void UTypedElementDatabaseCompatibility::PendingRegistration<AddressType>::Remov
 			// loop to avoid many resizes happening.
 			constexpr bool bAllowToShrink = false;
 			Addresses.RemoveAtSwap(AddressIt - AddressBegin, 1, bAllowToShrink);
-			StorageInterface.RemoveRow(*RowHandleIt);
 			ReservedRowHandles.RemoveAtSwap(RowHandleIt - RowHandleBegin, 1, bAllowToShrink);
+			StorageInterface.RemoveRow(*RowHandleIt);
 			--AddressEnd;
 		}
 	}
@@ -518,10 +518,9 @@ void UTypedElementDatabaseCompatibility::TickPendingUObjectRegistration()
 		for (auto It = UObjectsPendingRegistration.CreateIterator(); It; ++It)
 		{
 			It->Value.RemoveInvalidEntries(*Storage,
-				[this](const TWeakObjectPtr<UObject>& Object)
+				[](const TWeakObjectPtr<UObject>& Object)
 				{
-					UObject* Instance = Object.Get();
-					return Instance && Storage->IsRowAvailable(FindRowWithCompatibleObjectExplicit(Instance));
+					return Object.Get() != nullptr;
 				});
 		}
 
@@ -559,9 +558,9 @@ void UTypedElementDatabaseCompatibility::TickPendingExternalObjectRegistration()
 		for (auto It = ExternalObjectsPendingRegistration.CreateIterator(); It; ++It)
 		{
 			It->Value.RemoveInvalidEntries(*Storage,
-				[this](const ExternalObjectRegistration& Object)
+				[](const ExternalObjectRegistration& Object)
 				{
-					return Object.Object && Storage->IsRowAvailable(FindRowWithCompatibleObjectExplicit(Object.Object));
+					return Object.Object != nullptr;
 				});
 		}
 
