Update label to support editing.

The following changes were:
- FTEDSActorOutlinerTreeItem now sends the request to rename to TEDS.
- Introduced temporary "UI Capabilities" to TEDS UI. These are interfaces that case be used to discover the capabilities of a Slate Widget. This helps expose Slate functionality without having to directly tie to a specific widget.
- Actor name changes are now done through FActorLabelUtilities::RenameExistingActor to remain consistent with the original implementation.
- Actor name changes are now recorded as a transaction so undo/redo is supported.
- The label widget now supports a mode in which it can edit labels as well.
- The label widget has been reworked to use UI Capabilities to improve type safety.

diff --git a/Engine/Plugins/Experimental/TypedElementsDataStorage/Source/TypedElementsDataStorage/Private/Compatibility/TypedElementActorLabelQueries.cpp b/Engine/Plugins/Experimental/TypedElementsDataStorage/Source/TypedElementsDataStorage/Private/Compatibility/TypedElementActorLabelQueries.cpp
--- a/Engine/Plugins/Experimental/TypedElementsDataStorage/Source/TypedElementsDataStorage/Private/Compatibility/TypedElementActorLabelQueries.cpp
+++ b/Engine/Plugins/Experimental/TypedElementsDataStorage/Source/TypedElementsDataStorage/Private/Compatibility/TypedElementActorLabelQueries.cpp
@@ -2,6 +2,7 @@
 
 #include "Compatibility/TypedElementActorLabelQueries.h"
 
+#include "Editor/EditorEngine.h"
 #include "Elements/Columns/TypedElementMiscColumns.h"
 #include "Elements/Columns/TypedElementLabelColumns.h"
 #include "Elements/Framework/TypedElementQueryBuilder.h"
@@ -10,6 +11,9 @@
 #include "Hash/CityHash.h"
 #include "MassActorSubsystem.h"
 #include "ProfilingDebugging/CpuProfilerTrace.h"
+#include "ScopedTransaction.h"
+
+#define LOCTEXT_NAMESPACE "TypedElementDataStorage"
 
 FAutoConsoleCommandWithOutputDevice PrintActorLabelsConsoleCommand(
 	TEXT("TEDS.PrintActorLabels"),
@@ -123,7 +127,8 @@ void UTypedElementActorLabelFactory::RegisterLabelColumnToActorQuery(ITypedEleme
 					uint64 ActorLabelHash = CityHash64(reinterpret_cast<const char*>(*ActorLabel), ActorLabel.Len() * sizeof(**ActorLabel));
 					if (LabelHash.LabelHash != ActorLabelHash)
 					{
-						ActorInstance->SetActorLabel(Label.Label);
+						const FScopedTransaction Transaction(LOCTEXT("RenameActorTransaction", "Rename Actor"));
+						FActorLabelUtilities::RenameExistingActor(ActorInstance, Label.Label);
 					}
 				}
 			}
@@ -133,3 +138,5 @@ void UTypedElementActorLabelFactory::RegisterLabelColumnToActorQuery(ITypedEleme
 		.Compile()
 	);
 }
+
+#undef LOCTEXT_NAMESPACE
diff --git a/Engine/Plugins/Experimental/TypedElementsDataStorage/Source/TypedElementsDataStorageUI/Private/TypedElementLabelWidget.cpp b/Engine/Plugins/Experimental/TypedElementsDataStorage/Source/TypedElementsDataStorageUI/Private/TypedElementLabelWidget.cpp
--- a/Engine/Plugins/Experimental/TypedElementsDataStorage/Source/TypedElementsDataStorageUI/Private/TypedElementLabelWidget.cpp
+++ b/Engine/Plugins/Experimental/TypedElementsDataStorage/Source/TypedElementsDataStorageUI/Private/TypedElementLabelWidget.cpp
@@ -2,13 +2,18 @@
 
 #include "TypedElementLabelWidget.h"
 
+#include "ActorEditorUtils.h"
 #include "Elements/Columns/TypedElementLabelColumns.h"
 #include "Elements/Columns/TypedElementMiscColumns.h"
 #include "Elements/Columns/TypedElementSlateWidgetColumns.h"
 #include "Elements/Columns/TypedElementValueCacheColumns.h"
 #include "Elements/Framework/TypedElementQueryBuilder.h"
+#include "Elements/Interfaces/Capabilities/TypedElementUiEditableCapability.h"
+#include "Elements/Interfaces/Capabilities/TypedElementUiTextCapability.h"
+#include "Elements/Interfaces/Capabilities/TypedElementUiToolTipCapability.h"
 #include "TypedElementSubsystems.h"
 #include "Widgets/Text/STextBlock.h"
+#include "Widgets/Text/SInlineEditableTextBlock.h"
 
 #define LOCTEXT_NAMESPACE "TypedElementUI_LabelWidget"
 
@@ -19,22 +24,26 @@
 static void UpdateTextWidget(const TWeakPtr<SWidget>& Widget, const FTypedElementLabelColumn& Label, const uint64* HashValue)
 {
 	TSharedPtr<SWidget> WidgetPointer = Widget.Pin();
-	checkf(WidgetPointer, TEXT("Referenced widget is not valid. A constructed widget may not have been cleaned up. This can "
+	if (ensureMsgf(WidgetPointer, TEXT("Referenced widget is not valid. A constructed widget may not have been cleaned up. This can "
 		"also happen if this processor is running in the same phase as the processors responsible for cleaning up old "
-		"references."));
-	checkf(WidgetPointer->GetType() == STextBlock::StaticWidgetClass().GetWidgetType(),
-		TEXT("Stored widget with FTypedElementLabelWidgetConstructor doesn't match type %s, but was a %s."),
-		*(STextBlock::StaticWidgetClass().GetWidgetType().ToString()),
-		*(WidgetPointer->GetTypeAsString()));
-	STextBlock* WidgetInstance = static_cast<STextBlock*>(WidgetPointer.Get());
-	WidgetInstance->SetText(FText::FromString(Label.Label));
-	if (!HashValue)
+		"references.")))
 	{
-		WidgetInstance->SetToolTipText(FText::FromString(Label.Label));
-	}
-	else
-	{
-		WidgetInstance->SetToolTipText(FText::FromString(FString::Format(TEXT("{0}\nHash: {1}"), { Label.Label, *HashValue })));
+		if (TSharedPtr<ITypedElementUiTextCapability> Text = WidgetPointer->GetMetaData<ITypedElementUiTextCapability>())
+		{
+			Text->SetText(FText::FromString(Label.Label));
+		}
+
+		if (TSharedPtr<ITypedElementUiToolTipCapability> ToolTip = WidgetPointer->GetMetaData<ITypedElementUiToolTipCapability>())
+		{
+			if (!HashValue)
+			{
+				ToolTip->SetToolTipText(FText::FromString(Label.Label));
+			}
+			else
+			{
+				ToolTip->SetToolTipText(FText::FromString(FString::Format(TEXT("{0}\nHash: {1}"), { Label.Label, *HashValue })));
+			}
+		}
 	}
 }
 
@@ -152,17 +161,74 @@ TConstArrayView<const UScriptStruct*> FTypedElementLabelWidgetConstructor::GetAd
 	return Columns;
 }
 
-TSharedPtr<SWidget> FTypedElementLabelWidgetConstructor::CreateWidget(const TypedElementDataStorage::FMetaDataView& Arguments)
+TSharedPtr<SWidget> FTypedElementLabelWidgetConstructor::Construct(
+	TypedElementRowHandle Row,
+	ITypedElementDataStorageInterface* DataStorage,
+	ITypedElementDataStorageUiInterface* DataStorageUi,
+	const TypedElementDataStorage::FMetaDataView& Arguments)
 {
+	TSharedPtr<SWidget> Result;
 	const bool* IsEditable = Arguments.FindForColumn<FTypedElementLabelColumn>(TypedElementDataStorage::IsEditableName).TryGetExact<bool>();
-	return SNew(STextBlock).IsEnabled(IsEditable ? *IsEditable : false);
+	if (IsEditable && *IsEditable)
+	{
+		if (FTypedElementRowReferenceColumn* TargetRowColumn = DataStorage->GetColumn<FTypedElementRowReferenceColumn>(Row))
+		{
+			TSharedPtr<SInlineEditableTextBlock> TextBlock = SNew(SInlineEditableTextBlock)
+				.OnTextCommitted_Lambda(
+					[DataStorage, TargetRow = TargetRowColumn->Row](const FText& NewText, ETextCommit::Type CommitInfo)
+					{
+						// This callback happens on the game thread so it's safe to directly call into the data storage.
+						FString NewLabelText = NewText.ToString();
+						if (FTypedElementLabelHashColumn* LabelHashColumn = DataStorage->GetColumn<FTypedElementLabelHashColumn>(TargetRow))
+						{
+							LabelHashColumn->LabelHash = CityHash64(reinterpret_cast<const char*>(*NewLabelText), NewLabelText.Len() * sizeof(**NewLabelText));
+						}
+						if (FTypedElementLabelColumn* LabelColumn = DataStorage->GetColumn<FTypedElementLabelColumn>(TargetRow))
+						{
+							LabelColumn->Label = MoveTemp(NewLabelText);
+						}
+						DataStorage->AddColumn<FTypedElementSyncBackToWorldTag>(TargetRow);
+					})
+				.OnVerifyTextChanged_Lambda([](const FText& Label, FText& ErrorMessage)
+					{
+						// Note: The use of actor specific functionality should be minimized, but this function acts generic enough that the 
+						// use of actor is just in names.
+						return FActorEditorUtils::ValidateActorName(Label, ErrorMessage);
+					});
+			TextBlock->AddMetadata(MakeShared<TTypedElementUiEditableCapability<SInlineEditableTextBlock>>(*TextBlock));
+			TextBlock->AddMetadata(MakeShared<TTypedElementUiTextCapability<SInlineEditableTextBlock>>(*TextBlock));
+			TextBlock->AddMetadata(MakeShared<TTypedElementUiToolTipCapability<SInlineEditableTextBlock>>(*TextBlock));
+			Result = TextBlock;
+		}
+	}
+	else
+	{
+		TSharedPtr<STextBlock> TextBlock = SNew(STextBlock)
+			.IsEnabled(false);
+		TextBlock->AddMetadata(MakeShared<TTypedElementUiTextCapability<STextBlock>>(*TextBlock));
+		TextBlock->AddMetadata(MakeShared<TTypedElementUiToolTipCapability<STextBlock>>(*TextBlock));
+		Result = TextBlock;
+	}
+	
+	if (Result)
+	{
+		DataStorage->GetColumn<FTypedElementSlateWidgetReferenceColumn>(Row)->Widget = Result;
+		if (SetColumns(DataStorage, Row))
+		{
+			if (FinalizeWidget(DataStorage, DataStorageUi, Row, Result))
+			{
+				return Result;
+			}
+		}
+	}
+
+	return nullptr;
 }
 
 bool FTypedElementLabelWidgetConstructor::SetColumns(ITypedElementDataStorageInterface* DataStorage, TypedElementRowHandle Row)
 {
 	DataStorage->GetColumn<FTypedElementU64IntValueCacheColumn>(Row)->Value = 0;
-	DataStorage->GetColumn<FTypedElementLabelWidgetColumn>(Row)->bShowHashInTooltip = false;
-
+	DataStorage->AddOrGetColumn<FTypedElementLabelWidgetColumn>(Row)->bShowHashInTooltip = false;
 	return true;
 }
 
@@ -176,8 +242,7 @@ bool FTypedElementLabelWidgetConstructor::FinalizeWidget(
 		DataStorage, 
 		DataStorage->GetColumn<FTypedElementRowReferenceColumn>(Row)->Row,
 		*DataStorage->GetColumn<FTypedElementU64IntValueCacheColumn>(Row),
-		Widget,
-		false);
+		Widget, false);
 	return true;
 }
 
@@ -195,8 +260,7 @@ FTypedElementLabelWithHashTooltipWidgetConstructor::FTypedElementLabelWithHashTo
 bool FTypedElementLabelWithHashTooltipWidgetConstructor::SetColumns(ITypedElementDataStorageInterface* DataStorage, TypedElementRowHandle Row)
 {
 	DataStorage->GetColumn<FTypedElementU64IntValueCacheColumn>(Row)->Value = 0;
-	DataStorage->GetColumn<FTypedElementLabelWidgetColumn>(Row)->bShowHashInTooltip = true;
-
+	DataStorage->AddOrGetColumn<FTypedElementLabelWidgetColumn>(Row)->bShowHashInTooltip = true;
 	return true;
 }
 
@@ -210,8 +274,7 @@ bool FTypedElementLabelWithHashTooltipWidgetConstructor::FinalizeWidget(
 		DataStorage,
 		DataStorage->GetColumn<FTypedElementRowReferenceColumn>(Row)->Row,
 		*DataStorage->GetColumn<FTypedElementU64IntValueCacheColumn>(Row),
-		Widget,
-		true);
+		Widget, true);
 	return true;
 }
 
diff --git a/Engine/Plugins/Experimental/TypedElementsDataStorage/Source/TypedElementsDataStorageUI/Public/TypedElementLabelWidget.h b/Engine/Plugins/Experimental/TypedElementsDataStorage/Source/TypedElementsDataStorageUI/Public/TypedElementLabelWidget.h
--- a/Engine/Plugins/Experimental/TypedElementsDataStorage/Source/TypedElementsDataStorageUI/Public/TypedElementLabelWidget.h
+++ b/Engine/Plugins/Experimental/TypedElementsDataStorage/Source/TypedElementsDataStorageUI/Public/TypedElementLabelWidget.h
@@ -34,9 +34,14 @@ public:
 
 	TConstArrayView<const UScriptStruct*> GetAdditionalColumnsList() const override;
 
+	TSharedPtr<SWidget> Construct(
+		TypedElementRowHandle Row, 
+		ITypedElementDataStorageInterface* DataStorage,
+		ITypedElementDataStorageUiInterface* DataStorageUi,
+		const TypedElementDataStorage::FMetaDataView& Arguments) override;
+
 protected:
 	explicit FTypedElementLabelWidgetConstructor(const UScriptStruct* InTypeInfo);
-	TSharedPtr<SWidget> CreateWidget(const TypedElementDataStorage::FMetaDataView& Arguments) override;
 	bool SetColumns(ITypedElementDataStorageInterface* DataStorage, TypedElementRowHandle Row) override;
 	bool FinalizeWidget(ITypedElementDataStorageInterface* DataStorage, ITypedElementDataStorageUiInterface* DataStorageUi,
 		TypedElementRowHandle Row, const TSharedPtr<SWidget>& Widget) override;
@@ -62,5 +67,5 @@ struct TYPEDELEMENTSDATASTORAGEUI_API FTypedElementLabelWidgetColumn : public FT
 {
 	GENERATED_BODY()
 
-	bool bShowHashInTooltip{ false };
+	bool bShowHashInTooltip : 1{ false };
 };
\ No newline at end of file
diff --git a/Engine/Source/Runtime/TypedElementFramework/Public/Elements/Interfaces/Capabilities/TypedElementUiEditableCapability.h b/Engine/Source/Runtime/TypedElementFramework/Public/Elements/Interfaces/Capabilities/TypedElementUiEditableCapability.h
new file mode 100644
--- /dev/null
+++ b/Engine/Source/Runtime/TypedElementFramework/Public/Elements/Interfaces/Capabilities/TypedElementUiEditableCapability.h
@@ -0,0 +1,38 @@
+// Copyright Epic Games, Inc. All Rights Reserved.
+
+#pragma once
+
+#include "Elements/Interfaces/TypedElementUiCapabilities.h"
+
+/**
+ * Interface to provide access to widget that have explicit editing modes.
+ */
+class ITypedElementUiEditableCapability : public ITypedElementUiCapability
+{
+public:
+	SLATE_METADATA_TYPE(ITypedElementUiEditableCapability, ITypedElementUiCapability)
+
+	~ITypedElementUiEditableCapability() override = default;
+
+	virtual void EnterEditingMode() = 0;
+	virtual void ExitEditingMode() = 0;
+};
+
+template<typename WidgetType>
+class TTypedElementUiEditableCapability : public ITypedElementUiEditableCapability
+{
+public:
+	explicit TTypedElementUiEditableCapability(WidgetType& InWidget) : Widget(InWidget){}
+	
+	void EnterEditingMode() override 
+	{ 
+		Widget.EnterEditingMode();
+	}
+	void ExitEditingMode() override
+	{
+		Widget.ExitEditingMode();
+	}
+
+private:
+	WidgetType& Widget;
+};
\ No newline at end of file
diff --git a/Engine/Source/Runtime/TypedElementFramework/Public/Elements/Interfaces/Capabilities/TypedElementUiTextCapability.h b/Engine/Source/Runtime/TypedElementFramework/Public/Elements/Interfaces/Capabilities/TypedElementUiTextCapability.h
new file mode 100644
--- /dev/null
+++ b/Engine/Source/Runtime/TypedElementFramework/Public/Elements/Interfaces/Capabilities/TypedElementUiTextCapability.h
@@ -0,0 +1,35 @@
+// Copyright Epic Games, Inc. All Rights Reserved.
+
+#pragma once
+
+#include "Elements/Interfaces/TypedElementUiCapabilities.h"
+#include "Internationalization/Text.h"
+#include "Misc/Attribute.h"
+
+/**
+ * Interface to provide access to widgets that support working with text.
+ */
+class ITypedElementUiTextCapability : public ITypedElementUiCapability
+{
+public:
+	SLATE_METADATA_TYPE(ITypedElementUiTextCapability, ITypedElementUiCapability)
+
+	~ITypedElementUiTextCapability() override = default;
+
+	virtual void SetText(TAttribute<FText> Text) = 0;
+};
+
+template<typename WidgetType>
+class TTypedElementUiTextCapability : public ITypedElementUiTextCapability
+{
+public:
+	explicit TTypedElementUiTextCapability(WidgetType& InWidget) : Widget(InWidget){}
+	
+	void SetText(TAttribute<FText> Text) override
+	{
+		Widget.SetText(Text);
+	}
+
+private:
+	WidgetType& Widget;
+};
\ No newline at end of file
diff --git a/Engine/Source/Runtime/TypedElementFramework/Public/Elements/Interfaces/Capabilities/TypedElementUiTooltipCapability.h b/Engine/Source/Runtime/TypedElementFramework/Public/Elements/Interfaces/Capabilities/TypedElementUiTooltipCapability.h
new file mode 100644
--- /dev/null
+++ b/Engine/Source/Runtime/TypedElementFramework/Public/Elements/Interfaces/Capabilities/TypedElementUiTooltipCapability.h
@@ -0,0 +1,67 @@
+// Copyright Epic Games, Inc. All Rights Reserved.
+
+#pragma once
+
+#include "Elements/Interfaces/TypedElementUiCapabilities.h"
+#include "Internationalization/Text.h"
+#include "Misc/Attribute.h"
+#include "Widgets/IToolTip.h"
+
+/**
+ * Interface to provide access to tool tips on widgets.
+ */
+class ITypedElementUiToolTipCapability : public ITypedElementUiCapability
+{
+public:
+	SLATE_METADATA_TYPE(ITypedElementUiToolTipCapability, ITypedElementUiCapability)
+
+	~ITypedElementUiToolTipCapability() override = default;
+
+	virtual void SetToolTipText(const TAttribute<FText>& ToolTipText) = 0;
+	virtual void SetToolTipText(const FText& ToolTipText) = 0;
+	virtual void SetToolTip(const TAttribute<TSharedPtr<IToolTip>>& ToolTip) = 0;
+	virtual TSharedPtr<IToolTip> GetToolTip() = 0;
+
+	virtual void EnableToolTipForceField(const bool bEnableForceField) = 0;
+	virtual bool HasToolTipForceField() const = 0;
+};
+
+template<typename WidgetType>
+class TTypedElementUiToolTipCapability : public ITypedElementUiToolTipCapability
+{
+public:
+	explicit TTypedElementUiToolTipCapability(WidgetType& InWidget) : Widget(InWidget){}
+
+	void SetToolTipText(const TAttribute<FText>& ToolTipText) override
+	{
+		Widget.SetToolTipText(ToolTipText);	
+	}
+
+	void SetToolTipText(const FText& ToolTipText) override
+	{
+		Widget.SetToolTipText(ToolTipText);
+	}
+
+	void SetToolTip(const TAttribute<TSharedPtr<IToolTip>>& ToolTip) override
+	{
+		Widget.SetToolTip(ToolTip);
+	}
+
+	TSharedPtr<IToolTip> GetToolTip() override
+	{
+		return Widget.GetToolTip();
+	}
+
+	void EnableToolTipForceField(const bool bEnableForceField) override
+	{
+		Widget.EnableToolTipForceField(bEnableForceField);
+	}
+
+	bool HasToolTipForceField() const override
+	{
+		return Widget.HasToolTipForceField();
+	}
+
+private:
+	WidgetType& Widget;
+};
\ No newline at end of file
diff --git a/Engine/Source/Runtime/TypedElementFramework/Public/Elements/Interfaces/TypedElementUiCapabilities.h b/Engine/Source/Runtime/TypedElementFramework/Public/Elements/Interfaces/TypedElementUiCapabilities.h
new file mode 100644
--- /dev/null
+++ b/Engine/Source/Runtime/TypedElementFramework/Public/Elements/Interfaces/TypedElementUiCapabilities.h
@@ -0,0 +1,18 @@
+// Copyright Epic Games, Inc. All Rights Reserved.
+
+#pragma once
+
+#include "Types/ISlateMetaData.h"
+
+/**
+ * UI Capabilities are objects that can be stored as meta data on widgets created through TEDS UI and
+ * describe the functionality offered by the widget. This allows arbitrary widgets to describe themselves
+ * and allows users of TEDS UI to call into or register with the widget.
+ */
+class ITypedElementUiCapability : public ISlateMetaData
+{
+public:
+	SLATE_METADATA_TYPE(ITypedElementUiCapability, ISlateMetaData)
+
+	~ITypedElementUiCapability() override = default;
+};
\ No newline at end of file
