Fixed typo.

#nojira

diff --git a/Engine/Plugins/Experimental/TypedElementsDataStorage/Source/TypedElementsDataStorageUI/Private/TypedElementLabelWidget.cpp b/Engine/Plugins/Experimental/TypedElementsDataStorage/Source/TypedElementsDataStorageUI/Private/TypedElementLabelWidget.cpp
--- a/Engine/Plugins/Experimental/TypedElementsDataStorage/Source/TypedElementsDataStorageUI/Private/TypedElementLabelWidget.cpp
+++ b/Engine/Plugins/Experimental/TypedElementsDataStorage/Source/TypedElementsDataStorageUI/Private/TypedElementLabelWidget.cpp
@@ -10,7 +10,7 @@
 #include "Elements/Framework/TypedElementQueryBuilder.h"
 #include "Elements/Interfaces/Capabilities/TypedElementUiEditableCapability.h"
 #include "Elements/Interfaces/Capabilities/TypedElementUiTextCapability.h"
-#include "Elements/Interfaces/Capabilities/TypedElementUiToolTipCapability.h"
+#include "Elements/Interfaces/Capabilities/TypedElementUiTooltipCapability.h"
 #include "TypedElementSubsystems.h"
 #include "Widgets/Text/STextBlock.h"
 #include "Widgets/Text/SInlineEditableTextBlock.h"
@@ -33,7 +33,7 @@ static void UpdateTextWidget(const TWeakPtr<SWidget>& Widget, const FTypedElemen
 			Text->SetText(FText::FromString(Label.Label));
 		}
 
-		if (TSharedPtr<ITypedElementUiToolTipCapability> ToolTip = WidgetPointer->GetMetaData<ITypedElementUiToolTipCapability>())
+		if (TSharedPtr<ITypedElementUiTooltipCapability> ToolTip = WidgetPointer->GetMetaData<ITypedElementUiTooltipCapability>())
 		{
 			if (!HashValue)
 			{
@@ -197,7 +197,7 @@ TSharedPtr<SWidget> FTypedElementLabelWidgetConstructor::Construct(
 					});
 			TextBlock->AddMetadata(MakeShared<TTypedElementUiEditableCapability<SInlineEditableTextBlock>>(*TextBlock));
 			TextBlock->AddMetadata(MakeShared<TTypedElementUiTextCapability<SInlineEditableTextBlock>>(*TextBlock));
-			TextBlock->AddMetadata(MakeShared<TTypedElementUiToolTipCapability<SInlineEditableTextBlock>>(*TextBlock));
+			TextBlock->AddMetadata(MakeShared<TTypedElementUiTooltipCapability<SInlineEditableTextBlock>>(*TextBlock));
 			Result = TextBlock;
 		}
 	}
@@ -206,7 +206,7 @@ TSharedPtr<SWidget> FTypedElementLabelWidgetConstructor::Construct(
 		TSharedPtr<STextBlock> TextBlock = SNew(STextBlock)
 			.IsEnabled(false);
 		TextBlock->AddMetadata(MakeShared<TTypedElementUiTextCapability<STextBlock>>(*TextBlock));
-		TextBlock->AddMetadata(MakeShared<TTypedElementUiToolTipCapability<STextBlock>>(*TextBlock));
+		TextBlock->AddMetadata(MakeShared<TTypedElementUiTooltipCapability<STextBlock>>(*TextBlock));
 		Result = TextBlock;
 	}
 	
diff --git a/Engine/Source/Runtime/TypedElementFramework/Public/Elements/Interfaces/Capabilities/TypedElementUiTooltipCapability.h b/Engine/Source/Runtime/TypedElementFramework/Public/Elements/Interfaces/Capabilities/TypedElementUiTooltipCapability.h
--- a/Engine/Source/Runtime/TypedElementFramework/Public/Elements/Interfaces/Capabilities/TypedElementUiTooltipCapability.h
+++ b/Engine/Source/Runtime/TypedElementFramework/Public/Elements/Interfaces/Capabilities/TypedElementUiTooltipCapability.h
@@ -10,12 +10,12 @@
 /**
  * Interface to provide access to tool tips on widgets.
  */
-class ITypedElementUiToolTipCapability : public ITypedElementUiCapability
+class ITypedElementUiTooltipCapability : public ITypedElementUiCapability
 {
 public:
-	SLATE_METADATA_TYPE(ITypedElementUiToolTipCapability, ITypedElementUiCapability)
+	SLATE_METADATA_TYPE(ITypedElementUiTooltipCapability, ITypedElementUiCapability)
 
-	~ITypedElementUiToolTipCapability() override = default;
+	~ITypedElementUiTooltipCapability() override = default;
 
 	virtual void SetToolTipText(const TAttribute<FText>& ToolTipText) = 0;
 	virtual void SetToolTipText(const FText& ToolTipText) = 0;
@@ -27,10 +27,10 @@ public:
 };
 
 template<typename WidgetType>
-class TTypedElementUiToolTipCapability : public ITypedElementUiToolTipCapability
+class TTypedElementUiTooltipCapability : public ITypedElementUiTooltipCapability
 {
 public:
-	explicit TTypedElementUiToolTipCapability(WidgetType& InWidget) : Widget(InWidget){}
+	explicit TTypedElementUiTooltipCapability(WidgetType& InWidget) : Widget(InWidget){}
 
 	void SetToolTipText(const TAttribute<FText>& ToolTipText) override
 	{
