Fix TEDS invalid entity ensure in case where entity is destroyed before a delayed operation ticks

diff --git a/Engine/Plugins/Experimental/TypedElementsDataStorage/Source/TypedElementsDataStorage/Private/Processors/TypedElementProcessorAdaptors.cpp b/Engine/Plugins/Experimental/TypedElementsDataStorage/Source/TypedElementsDataStorage/Private/Processors/TypedElementProcessorAdaptors.cpp
--- a/Engine/Plugins/Experimental/TypedElementsDataStorage/Source/TypedElementsDataStorage/Private/Processors/TypedElementProcessorAdaptors.cpp
+++ b/Engine/Plugins/Experimental/TypedElementsDataStorage/Source/TypedElementsDataStorage/Private/Processors/TypedElementProcessorAdaptors.cpp
@@ -279,20 +279,28 @@ public:
 			Context.Defer().PushCommand<FMassDeferredAddCommand>(
 				[AddedColumn](FMassEntityManager& System)
 				{
-					// Check before adding.  Mass's AddFragmentToEntity is not idempotent and will assert if adding
-					// column to a row that already has one
-					FStructView Fragment = System.GetFragmentDataStruct(AddedColumn->Entity, AddedColumn->FragmentType);
-					if (!Fragment.IsValid())
+					// Check entity before proceeding. It's possible it may have been invalidated before this defer call fired
+					if (System.IsEntityActive(AddedColumn->Entity))
 					{
-						System.AddFragmentToEntity(AddedColumn->Entity, AddedColumn->FragmentType);
+						// Check before adding.  Mass's AddFragmentToEntity is not idempotent and will assert if adding
+						// column to a row that already has one
+						FStructView Fragment = System.GetFragmentDataStruct(AddedColumn->Entity, AddedColumn->FragmentType);
+						if (!Fragment.IsValid())
+						{
+							System.AddFragmentToEntity(AddedColumn->Entity, AddedColumn->FragmentType);
+						}
 					}
 				});
 			
 			Context.Defer().PushCommand<FMassDeferredSetCommand>(
 				[AddedColumn](FMassEntityManager& System)
 				{
-					FStructView Fragment = System.GetFragmentDataStruct(AddedColumn->Entity, AddedColumn->FragmentType);
-					AddedColumn->FragmentType->CopyScriptStruct(Fragment.GetMemory(), AddedColumn->Object);
+					// Check entity before proceeding. It's possible it may have been invalidated before this defer call fired
+					if (System.IsEntityActive(AddedColumn->Entity))
+					{
+						FStructView Fragment = System.GetFragmentDataStruct(AddedColumn->Entity, AddedColumn->FragmentType);
+						AddedColumn->FragmentType->CopyScriptStruct(Fragment.GetMemory(), AddedColumn->Object);
+					}
 				});
 		}
 		else if (ObjectType->IsChildOf(FMassTag::StaticStruct()))
@@ -300,7 +308,11 @@ public:
 			Context.Defer().PushCommand<FMassDeferredAddCommand>(
 				[AddedColumn](FMassEntityManager& System)
 				{
-					System.AddTagToEntity(AddedColumn->Entity, AddedColumn->FragmentType);
+					// Check entity before proceeding. It's possible it may have been invalidated before this defer call fired
+					if (System.IsEntityActive(AddedColumn->Entity))
+					{
+						System.AddTagToEntity(AddedColumn->Entity, AddedColumn->FragmentType);
+					}
 				});
 		}
 
