[TEDS] Move outliner widgets to the outliner subfolder of the data storage UI

diff --git a/Engine/Plugins/Experimental/TypedElementsDataStorage/Source/TypedElementsDataStorageUI/Private/TypedElementOutliner/Widgets/TypedElementTransformHeadsUpWidget.cpp b/Engine/Plugins/Experimental/TypedElementsDataStorage/Source/TypedElementsDataStorageUI/Private/TypedElementOutliner/Widgets/TypedElementTransformHeadsUpWidget.cpp
new file mode 100644
--- /dev/null
+++ b/Engine/Plugins/Experimental/TypedElementsDataStorage/Source/TypedElementsDataStorageUI/Private/TypedElementOutliner/Widgets/TypedElementTransformHeadsUpWidget.cpp
@@ -0,0 +1,239 @@
+// Copyright Epic Games, Inc. All Rights Reserved.
+
+#include "TypedElementTransformHeadsUpWidget.h"
+
+#include "Elements/Columns/TypedElementMiscColumns.h"
+#include "Elements/Columns/TypedElementSlateWidgetColumns.h"
+#include "Elements/Columns/TypedElementTransformColumns.h"
+#include "Elements/Framework/TypedElementQueryBuilder.h"
+#include "TypedElementSubsystems.h"
+#include "Widgets/SCompoundWidget.h"
+#include "Widgets/SBoxPanel.h"
+#include "Widgets/Images/SImage.h"
+#include "Styling/AppStyle.h"
+#include "Widgets/Input/SNumericEntryBox.h"
+
+#define LOCTEXT_NAMESPACE "TypedElementUI_TransformHeadsUpWidget"
+
+namespace InternalTransformHelpers
+{
+	enum class EAbnormalTransformTypes : int32
+	{
+		None = 0x0000,
+		NonUniformScale = 0x0001,
+		NegativeXScale = 0x0002,
+		NegativeYScale = 0x0004,
+		NegativeZScale = 0x0008,
+		UnnormalizedRotation = 0x0010,
+	};
+	ENUM_CLASS_FLAGS(EAbnormalTransformTypes);
+
+	EAbnormalTransformTypes GetAbnormalTransformTypes(const FTransform& InTransform)
+	{
+		EAbnormalTransformTypes AbnormalTransformFlags = EAbnormalTransformTypes::None;
+		const FVector& TransformScale = InTransform.GetScale3D();
+
+		if (!TransformScale.GetAbs().AllComponentsEqual())
+		{
+			AbnormalTransformFlags |= EAbnormalTransformTypes::NonUniformScale;
+		}
+		if (TransformScale.X < 0.0)
+		{
+			AbnormalTransformFlags |= EAbnormalTransformTypes::NegativeXScale;
+		}
+		if (TransformScale.Y < 0.0)
+		{
+			AbnormalTransformFlags |= EAbnormalTransformTypes::NegativeYScale;
+		}
+		if (TransformScale.Z < 0.0)
+		{
+			AbnormalTransformFlags |= EAbnormalTransformTypes::NegativeZScale;
+		}
+		if (!InTransform.IsRotationNormalized())
+		{
+			AbnormalTransformFlags |= EAbnormalTransformTypes::UnnormalizedRotation;
+		}
+
+		return AbnormalTransformFlags;
+	}
+} // namespace InternalTransformHelpers
+
+class STransformQuickDisplay : public SHorizontalBox
+{
+	SLATE_DECLARE_WIDGET(STransformQuickDisplay, SHorizontalBox)
+public:
+	SLATE_BEGIN_ARGS(STransformQuickDisplay) {}
+	SLATE_END_ARGS()
+
+	void Construct(const FArguments& InArgs)
+	{
+		const static FMargin IconPadding(1, 1, 0, 0);
+		using namespace InternalTransformHelpers;
+
+		AddSlot().AutoWidth().Padding(IconPadding)
+		[
+			SNew(SImage)
+				.Image(FAppStyle::GetBrush("EditorViewport.ScaleGridSnap"))
+				.ToolTipText(LOCTEXT("NonUniformScaleTooltip", "Has Non-Uniform Scale"))
+				.Visibility_Lambda([this]() { return EnumHasAllFlags(AbnormalTransformFlags, EAbnormalTransformTypes::NonUniformScale) ? EVisibility::Visible : EVisibility::Collapsed; })
+		];
+		AddSlot().AutoWidth().Padding(IconPadding)
+		[
+			SNew(SImage)
+				.Image(FAppStyle::GetBrush("EditorViewport.ScaleMode"))
+				.ColorAndOpacity(SNumericEntryBox<double>::RedLabelBackgroundColor)
+				.ToolTipText(LOCTEXT("NegativeXScaleTooltip", "Has Negative X Scale"))
+				.Visibility_Lambda([this]() { return EnumHasAllFlags(AbnormalTransformFlags, EAbnormalTransformTypes::NegativeXScale) ? EVisibility::Visible : EVisibility::Collapsed; })
+		];
+		AddSlot().AutoWidth().Padding(IconPadding)
+		[
+			SNew(SImage)
+				.Image(FAppStyle::GetBrush("EditorViewport.ScaleMode"))
+				.ColorAndOpacity(SNumericEntryBox<double>::GreenLabelBackgroundColor)
+				.ToolTipText(LOCTEXT("NegativeYScaleTooltip", "Has Negative Y Scale"))
+				.Visibility_Lambda([this]() { return EnumHasAllFlags(AbnormalTransformFlags, EAbnormalTransformTypes::NegativeYScale) ? EVisibility::Visible : EVisibility::Collapsed; })
+		];
+		AddSlot().AutoWidth().Padding(IconPadding)
+		[
+			SNew(SImage)
+				.Image(FAppStyle::GetBrush("EditorViewport.ScaleMode"))
+				.ColorAndOpacity(SNumericEntryBox<double>::BlueLabelBackgroundColor)
+				.ToolTipText(LOCTEXT("NegativeZScaleTooltip", "Has Negative Z Scale"))
+				.Visibility_Lambda([this]() { return EnumHasAllFlags(AbnormalTransformFlags, EAbnormalTransformTypes::NegativeZScale) ? EVisibility::Visible : EVisibility::Collapsed; })
+		];
+		AddSlot().AutoWidth().Padding(IconPadding)
+		[
+			SNew(SImage)
+				.Image(FAppStyle::GetBrush("SurfaceDetails.AntiClockwiseRotation"))
+				.ToolTipText(LOCTEXT("UnnormalizedRotationTooltip", "Has Un-normalized Rotation"))
+				.Visibility_Lambda([this]() { return EnumHasAllFlags(AbnormalTransformFlags, EAbnormalTransformTypes::UnnormalizedRotation) ? EVisibility::Visible : EVisibility::Collapsed; })
+		];
+		AddSlot().AutoWidth().Padding(IconPadding)
+		[
+			SNew(SImage)
+				.Image(FAppStyle::GetBrush("Symbols.Check"))
+				.ToolTipText(LOCTEXT("NothingToReportTooltip", "No Abnormal Transform Data"))
+				.Visibility_Lambda([this]() { return (AbnormalTransformFlags == EAbnormalTransformTypes::None) ? EVisibility::Visible : EVisibility::Collapsed; })
+		];
+	}
+
+	void UpdateFromTransform(InternalTransformHelpers::EAbnormalTransformTypes InAbnormalTransformFlags)
+	{
+		if (AbnormalTransformFlags != InAbnormalTransformFlags)
+		{
+			AbnormalTransformFlags = InAbnormalTransformFlags;
+			Invalidate(EInvalidateWidgetReason::Visibility);
+		}
+	}
+
+private:
+	InternalTransformHelpers::EAbnormalTransformTypes AbnormalTransformFlags;
+};
+
+SLATE_IMPLEMENT_WIDGET(STransformQuickDisplay)
+void STransformQuickDisplay::PrivateRegisterAttributes(FSlateAttributeInitializer&) {}
+
+static void UpdateTransformHeadsUpDisplay(FTypedElementSlateWidgetReferenceColumn& Widget, InternalTransformHelpers::EAbnormalTransformTypes AbnormalTransformFlags)
+{
+	TSharedPtr<SWidget> WidgetPointer = Widget.Widget.Pin();
+	checkf(WidgetPointer, TEXT("Referenced widget is not valid. A constructed widget may not have been cleaned up. This can "
+		"also happen if this processor is running in the same phase as the processors responsible for cleaning up old "
+		"references."));
+	checkf(WidgetPointer->GetType() == STransformQuickDisplay::StaticWidgetClass().GetWidgetType(),
+		TEXT("Stored widget with FTypedElementTransformHeadsUpWidgetTag doesn't match type %s, but was a %s."),
+		*(STransformQuickDisplay::StaticWidgetClass().GetWidgetType().ToString()),
+		*(WidgetPointer->GetTypeAsString()));
+
+	STransformQuickDisplay* BoxWidget = static_cast<STransformQuickDisplay*>(WidgetPointer.Get());
+	BoxWidget->UpdateFromTransform(AbnormalTransformFlags);
+}
+
+//
+// UTypedElementTransformHeadsUpWidgetFactory
+//
+
+void UTypedElementTransformHeadsUpWidgetFactory::RegisterQueries(ITypedElementDataStorageInterface& DataStorage) const
+{
+	using namespace TypedElementQueryBuilder;
+	using DSI = ITypedElementDataStorageInterface;
+	namespace DS = TypedElementDataStorage;
+		
+	TypedElementQueryHandle UpdateTransformWidget = DataStorage.RegisterQuery(
+		Select()
+			.ReadOnly<FTypedElementLocalTransformColumn>()
+		.Where()
+			.Any<FTypedElementSyncFromWorldTag, FTypedElementSyncBackToWorldTag>()
+		.Compile());
+
+	DataStorage.RegisterQuery(
+		Select(TEXT("Sync Transform column to heads up display"),
+		FProcessor(DSI::EQueryTickPhase::FrameEnd, DataStorage.GetQueryTickGroupName(DSI::EQueryTickGroups::SyncWidgets))
+			.ForceToGameThread(true),
+			[](DS::IQueryContext& Context,
+				FTypedElementSlateWidgetReferenceColumn& Widget,
+				const FTypedElementRowReferenceColumn& ReferenceColumn)
+			{
+				Context.RunSubquery(0, ReferenceColumn.Row, CreateSubqueryCallbackBinding(
+					[&Widget](const FTypedElementLocalTransformColumn& Transform)
+					{
+						UpdateTransformHeadsUpDisplay(Widget, InternalTransformHelpers::GetAbnormalTransformTypes(Transform.Transform));
+					}));
+			}
+		)
+	.Where()
+		.All<FTypedElementTransformHeadsUpWidgetTag>()
+	.DependsOn()
+		.SubQuery(UpdateTransformWidget)
+	.Compile());
+
+}
+
+void UTypedElementTransformHeadsUpWidgetFactory::RegisterWidgetConstructors(ITypedElementDataStorageInterface& DataStorage,
+	ITypedElementDataStorageUiInterface& DataStorageUi) const
+{
+	using namespace TypedElementQueryBuilder;
+
+	DataStorageUi.RegisterWidgetFactory<FTypedElementTransformHeadsUpWidgetConstructor>(FName(TEXT("SceneOutliner.Cell")),
+		FColumn<FTypedElementLocalTransformColumn>());
+}
+
+
+
+//
+// FTypedElementTransformHeadsUpWidgetConstructor
+//
+
+FTypedElementTransformHeadsUpWidgetConstructor::FTypedElementTransformHeadsUpWidgetConstructor()
+	: Super(FTypedElementTransformHeadsUpWidgetConstructor::StaticStruct())
+{
+}
+
+TConstArrayView<const UScriptStruct*> FTypedElementTransformHeadsUpWidgetConstructor::GetAdditionalColumnsList() const
+{
+	static TTypedElementColumnTypeList<FTypedElementRowReferenceColumn,
+		FTypedElementTransformHeadsUpWidgetTag> Columns;
+	return Columns;
+}
+
+TSharedPtr<SWidget> FTypedElementTransformHeadsUpWidgetConstructor::CreateWidget(const TypedElementDataStorage::FMetaDataView& Arguments)
+{
+	return SNew(STransformQuickDisplay);
+}
+
+bool FTypedElementTransformHeadsUpWidgetConstructor::FinalizeWidget(
+	ITypedElementDataStorageInterface* DataStorage,
+	ITypedElementDataStorageUiInterface* DataStorageUi,
+	TypedElementRowHandle Row,
+	const TSharedPtr<SWidget>& Widget)
+{
+	FTypedElementRowReferenceColumn& RefColumn = *DataStorage->GetColumn<FTypedElementRowReferenceColumn>(Row);
+	if (const FTypedElementLocalTransformColumn* TransformColumn = DataStorage->GetColumn<FTypedElementLocalTransformColumn>(RefColumn.Row))
+	{
+		UpdateTransformHeadsUpDisplay(*DataStorage->GetColumn<FTypedElementSlateWidgetReferenceColumn>(Row), 
+			InternalTransformHelpers::GetAbnormalTransformTypes(TransformColumn->Transform));
+	}
+
+	return true;
+}
+
+#undef LOCTEXT_NAMESPACE
diff --git a/Engine/Plugins/Experimental/TypedElementsDataStorage/Source/TypedElementsDataStorageUI/Private/TypedElementOutliner/Widgets/TypedElementTypeInfoWidget.cpp b/Engine/Plugins/Experimental/TypedElementsDataStorage/Source/TypedElementsDataStorageUI/Private/TypedElementOutliner/Widgets/TypedElementTypeInfoWidget.cpp
new file mode 100644
--- /dev/null
+++ b/Engine/Plugins/Experimental/TypedElementsDataStorage/Source/TypedElementsDataStorageUI/Private/TypedElementOutliner/Widgets/TypedElementTypeInfoWidget.cpp
@@ -0,0 +1,159 @@
+// Copyright Epic Games, Inc. All Rights Reserved.
+
+#include "TypedElementTypeInfoWidget.h"
+
+#include "MassActorSubsystem.h"
+#include "SceneOutlinerHelpers.h"
+#include "Elements/Columns/TypedElementCompatibilityColumns.h"
+#include "Elements/Columns/TypedElementMiscColumns.h"
+#include "Elements/Columns/TypedElementSlateWidgetColumns.h"
+#include "Elements/Columns/TypedElementTypeInfoColumns.h"
+#include "Styling/SlateIconFinder.h"
+#include "Widgets/Text/STextBlock.h"
+#include "Widgets/Images/SImage.h"
+
+void UTypedElementTypeInfoWidgetFactory::RegisterWidgetConstructors(ITypedElementDataStorageInterface& DataStorage,
+	ITypedElementDataStorageUiInterface& DataStorageUi) const
+{
+	DataStorageUi.RegisterWidgetFactory<FTypedElementTypeInfoWidgetConstructor>(FName(TEXT("General.Cell")),
+		TypedElementQueryBuilder::FColumn<FTypedElementClassTypeInfoColumn>());
+
+}
+
+TMap<FName, const FSlateBrush*> FTypedElementTypeInfoWidgetConstructor::CachedIconMap;
+
+FTypedElementTypeInfoWidgetConstructor::FTypedElementTypeInfoWidgetConstructor()
+	: Super(FTypedElementTypeInfoWidgetConstructor::StaticStruct())
+	, bUseIcon(false)
+{
+}
+
+FTypedElementTypeInfoWidgetConstructor::FTypedElementTypeInfoWidgetConstructor(const UScriptStruct* InTypeInfo)
+	: Super(InTypeInfo)
+	, bUseIcon(false)
+{
+	
+}
+
+TSharedPtr<SWidget> FTypedElementTypeInfoWidgetConstructor::CreateWidget(
+	const TypedElementDataStorage::FMetaDataView& Arguments	)
+{
+	// Check if the caller provided metadata to use an icon widget
+	TypedElementDataStorage::FMetaDataEntryView MetaDataEntryView = Arguments.FindGeneric("TypedElementTypeInfoWidget_bUseIcon");
+	if(MetaDataEntryView.IsSet())
+	{
+		check(MetaDataEntryView.IsType<bool>());
+
+		bUseIcon = *MetaDataEntryView.TryGetExact<bool>();
+	}
+
+	if(bUseIcon)
+	{
+		return SNew(SImage)
+					.DesiredSizeOverride(FVector2D(16.f, 16.f))
+					.ColorAndOpacity(FSlateColor::UseForeground());
+	}
+	else
+	{
+		return SNew(SHorizontalBox);
+	}
+}
+
+bool FTypedElementTypeInfoWidgetConstructor::FinalizeWidget(ITypedElementDataStorageInterface* DataStorage,
+	ITypedElementDataStorageUiInterface* DataStorageUi, TypedElementRowHandle Row, const TSharedPtr<SWidget>& Widget)
+{
+	checkf(Widget, TEXT("Referenced widget is not valid. A constructed widget may not have been cleaned up. This can "
+		"also happen if this processor is running in the same phase as the processors responsible for cleaning up old "
+		"references."));
+
+	TypedElementRowHandle TargetRow = DataStorage->GetColumn<FTypedElementRowReferenceColumn>(Row)->Row;
+
+	if (const FTypedElementClassTypeInfoColumn* TypeInfoColumn = DataStorage->GetColumn<FTypedElementClassTypeInfoColumn>(TargetRow))
+	{
+		if(bUseIcon)
+		{
+			checkf(Widget->GetType() == SImage::StaticWidgetClass().GetWidgetType(),
+				TEXT("Stored widget with FTypedElementLabelWidgetConstructor doesn't match type %s, but was a %s."),
+				*(SImage::StaticWidgetClass().GetWidgetType().ToString()),
+				*(Widget->GetTypeAsString()));
+
+			SImage* WidgetInstance = static_cast<SImage*>(Widget.Get());
+			
+			WidgetInstance->SetImage(GetIconForRow(DataStorage, Row, TypeInfoColumn));
+		}
+		else
+		{
+			checkf(Widget->GetType() == SHorizontalBox::StaticWidgetClass().GetWidgetType(),
+				TEXT("Stored widget with FTypedElementLabelWidgetConstructor doesn't match type %s, but was a %s."),
+				*(SHorizontalBox::StaticWidgetClass().GetWidgetType().ToString()),
+				*(Widget->GetTypeAsString()));
+				
+			SHorizontalBox* WidgetInstance = static_cast<SHorizontalBox*>(Widget.Get());
+
+			TSharedPtr<SWidget> ActualWidget;
+
+			// Check if we have a hyperlink for this object
+			if (const FTypedElementUObjectColumn* ObjectColumn = DataStorage->GetColumn<FTypedElementUObjectColumn>(TargetRow))
+			{
+				ActualWidget = SceneOutliner::FSceneOutlinerHelpers::GetClassHyperlink(ObjectColumn->Object.Get());
+			}
+
+			// If not, we simply show a text block with the type
+			if(!ActualWidget)
+			{
+				ActualWidget = SNew(STextBlock)
+								.ColorAndOpacity(FSlateColor::UseSubduedForeground())
+								.Text(FText::FromString(TypeInfoColumn->TypeInfo.Get()->GetName()));
+			}
+
+
+			WidgetInstance->AddSlot()
+						.AutoWidth()
+						.VAlign(VAlign_Center)
+						.Padding(8, 0, 0, 0)
+						[
+							ActualWidget.ToSharedRef()
+						];
+		}
+	}
+
+	return true;
+}
+
+const FSlateBrush* FTypedElementTypeInfoWidgetConstructor::GetIconForRow(ITypedElementDataStorageInterface* DataStorage, TypedElementRowHandle Row, const FTypedElementClassTypeInfoColumn* TypeInfoColumn)
+{
+	/* The logic here is very similar to SActorTreeLabel::GetIcon in ActorTreeItem.cpp which allows the actor to specify
+	 * an override for the icon, and has a fallback to the class icon if not.
+	 */
+	
+	const UClass* Type = TypeInfoColumn->TypeInfo.Get();
+	FName IconName = Type->GetFName();
+
+	// Allow the actor the first chance to provide an icon override
+	FMassActorFragment* ActorStore = DataStorage->GetColumn<FMassActorFragment>(Row);
+	if(ActorStore)
+	{
+		if(const AActor* Actor = ActorStore->Get())
+		{
+			IconName = Actor->GetCustomIconName();
+		}
+	}
+
+	// Check the cache if we already found an icon for this class
+	if(const FSlateBrush** CachedBrush = CachedIconMap.Find(IconName))
+	{
+		if(*CachedBrush)
+		{
+			return *CachedBrush;
+		}
+	}
+
+	if(const FSlateBrush* FoundSlateBrush = FSlateIconFinder::FindIconForClass(Type).GetOptionalIcon())
+	{
+		CachedIconMap.Add(IconName, FoundSlateBrush);
+		return FoundSlateBrush;
+	}
+
+	// Fallback to the regular actor icon if we haven't found any specific icon
+	return FSlateIconFinder::FindIconForClass(AActor::StaticClass()).GetOptionalIcon();;
+}
diff --git a/Engine/Plugins/Experimental/TypedElementsDataStorage/Source/TypedElementsDataStorageUI/Private/TypedElementTransformHeadsUpWidget.cpp b/Engine/Plugins/Experimental/TypedElementsDataStorage/Source/TypedElementsDataStorageUI/Private/TypedElementTransformHeadsUpWidget.cpp
deleted file mode 100644
diff --git a/Engine/Plugins/Experimental/TypedElementsDataStorage/Source/TypedElementsDataStorageUI/Private/TypedElementTypeInfoWidget.cpp b/Engine/Plugins/Experimental/TypedElementsDataStorage/Source/TypedElementsDataStorageUI/Private/TypedElementTypeInfoWidget.cpp
deleted file mode 100644
