TypedElementTypeInfoWidget: Change member functions to protected instead of private

diff --git a/Engine/Plugins/Experimental/TypedElementsDataStorage/Source/TypedElementsDataStorageUI/Public/TypedElementTypeInfoWidget.h b/Engine/Plugins/Experimental/TypedElementsDataStorage/Source/TypedElementsDataStorageUI/Public/TypedElementTypeInfoWidget.h
--- a/Engine/Plugins/Experimental/TypedElementsDataStorage/Source/TypedElementsDataStorageUI/Public/TypedElementTypeInfoWidget.h
+++ b/Engine/Plugins/Experimental/TypedElementsDataStorage/Source/TypedElementsDataStorageUI/Public/TypedElementTypeInfoWidget.h
@@ -38,12 +38,12 @@ protected:
 	bool FinalizeWidget(ITypedElementDataStorageInterface* DataStorage, ITypedElementDataStorageUiInterface* DataStorageUi,
 		TypedElementRowHandle Row, const TSharedPtr<SWidget>& Widget) override;
 
-private:
+protected:
 
 	// Get the icon for a given row, checking the cache to see if we already have one
 	static const FSlateBrush* GetIconForRow(ITypedElementDataStorageInterface* DataStorage, TypedElementRowHandle Row, const FTypedElementClassTypeInfoColumn* TypeInfoColumn);
 	
-private:
+protected:
 
 	// Cache to avoid looking up the icon for a class every time
 	static TMap<FName, const FSlateBrush*> CachedIconMap;
