Added Indexing support to TEDS.

This change introduces the first step towards the ability to index columns for quick lookup. The following changes were made:
- Expanded the data storage API with calls to add, replace or remove an index for a row.
- Added a utility function to generate an index hash with. Currently it's limited to pointers but can be expanded upon as needed. Using other functions to create a hash or id with will also be valid.
- The context provided to query callbacks can now find a row from an index.
- Moved the scratch buffer and the new index table into an environment to make it easier to share among the various parts in TEDS. In the future this will be extended with a command buffer and the query store.
- TEDS Compat now uses indexing to find UObjects and regular objects instead of maintaining it's own table. This also means this can be used inside query callbacks to retrieve the row for objects quickly.
- The context provided to all types of query callback can now be used to verify if a row is available or initialized.

diff --git a/Engine/Plugins/Experimental/TypedElementsDataStorage/Source/TypedElementsDataStorage/Private/Processors/TypedElementProcessorAdaptors.cpp b/Engine/Plugins/Experimental/TypedElementsDataStorage/Source/TypedElementsDataStorage/Private/Processors/TypedElementProcessorAdaptors.cpp
--- a/Engine/Plugins/Experimental/TypedElementsDataStorage/Source/TypedElementsDataStorage/Private/Processors/TypedElementProcessorAdaptors.cpp
+++ b/Engine/Plugins/Experimental/TypedElementsDataStorage/Source/TypedElementsDataStorage/Private/Processors/TypedElementProcessorAdaptors.cpp
@@ -9,13 +9,27 @@
 #include "MassExecutionContext.h"
 #include "Queries/TypedElementExtendedQueryStore.h"
 #include "TypedElementDatabase.h"
-#include "TypedElementDatabaseScratchBuffer.h"
+#include "TypedElementDatabaseEnvironment.h"
 
 template<typename T>
 struct FMassContextCommon : public T
 {
 	~FMassContextCommon() override = default;
 
+	uint32 GetRowCount() const override
+	{
+		return Context.GetNumEntities();
+	}
+
+	TConstArrayView<TypedElementRowHandle> GetRowHandles() const override
+	{
+		static_assert(
+			sizeof(TypedElementRowHandle) == sizeof(FMassEntityHandle) && alignof(TypedElementRowHandle) == alignof(FMassEntityHandle),
+			"TypedElementRowHandle and FMassEntityHandle need to by layout compatible to support Typed Elements Data Storage.");
+		TConstArrayView<FMassEntityHandle> Entities = Context.GetEntities();
+		return TConstArrayView<TypedElementRowHandle>(reinterpret_cast<const TypedElementRowHandle*>(Entities.GetData()), Entities.Num());
+	}
+	
 	const void* GetColumn(const UScriptStruct* ColumnType) const override
 	{
 		return Context.GetFragmentView(ColumnType).GetData();
@@ -53,20 +67,6 @@ struct FMassContextCommon : public T
 		}
 	}
 
-	uint32 GetRowCount() const override
-	{
-		return Context.GetNumEntities();
-	}
-	
-	TConstArrayView<TypedElementRowHandle> GetRowHandles() const override
-	{
-		static_assert(
-			sizeof(TypedElementRowHandle) == sizeof(FMassEntityHandle) && alignof(TypedElementRowHandle) == alignof(FMassEntityHandle),
-			"TypedElementRowHandle and FMassEntityHandle need to by layout compatible to support Typed Elements Data Storage.");
-		TConstArrayView<FMassEntityHandle> Entities = Context.GetEntities();
-		return TConstArrayView<TypedElementRowHandle>(reinterpret_cast<const TypedElementRowHandle*>(Entities.GetData()), Entities.Num());
-	}
-
 protected:
 	explicit FMassContextCommon(FMassExecutionContext& InContext)
 		: Context(InContext)
@@ -170,11 +170,11 @@ private:
 public:
 
 	FMassContextForwarder(ITypedElementDataStorageInterface::FQueryDescription& InQueryDescription, FMassExecutionContext& InContext, 
-		FTypedElementExtendedQueryStore& InQueryStore, FTypedElementDatabaseScratchBuffer& InScratchBuffer)
+		FTypedElementExtendedQueryStore& InQueryStore, FTypedElementDatabaseEnvironment& InEnvironment)
 		: FMassContextCommon(InContext)
 		, QueryDescription(InQueryDescription)
 		, QueryStore(InQueryStore)
-		, ScratchBuffer(InScratchBuffer)
+		, Environment(InEnvironment)
 	{}
 
 	~FMassContextForwarder() override = default;
@@ -214,6 +214,16 @@ public:
 		}
 	}
 
+	bool IsRowAvailable(TypedElementDataStorage::RowHandle Row) const override
+	{
+		return Environment.GetMassEntityManager().IsEntityValid(FMassEntityHandle::FromNumber(Row));
+	}
+
+	bool HasRowBeenAssigned(TypedElementDataStorage::RowHandle Row) const override
+	{
+		return Environment.GetMassEntityManager().IsEntityActive(FMassEntityHandle::FromNumber(Row));
+	}
+
 	void RemoveRow(TypedElementRowHandle Row) override
 	{
 		Context.Defer().DestroyEntity(FMassEntityHandle::FromNumber(Row));
@@ -263,6 +273,7 @@ public:
 			}
 		};
 
+		FTypedElementDatabaseScratchBuffer& ScratchBuffer = Environment.GetScratchBuffer();
 		void* ObjectCopy = ScratchBuffer.Allocate(ObjectType->GetStructureSize(), ObjectType->GetMinAlignment());
 		FAddValueColumn* AddedColumn = nullptr;
 		if (ObjectType->StructFlags & (STRUCT_IsPlainOldData | STRUCT_NoDestructor))
@@ -352,6 +363,7 @@ public:
 			}
 		};
 
+		FTypedElementDatabaseScratchBuffer& ScratchBuffer = Environment.GetScratchBuffer();
 		void* MovedObject = ScratchBuffer.Allocate(ObjectType->GetStructureSize(), ObjectType->GetMinAlignment());
 		FAddMoveableValueColumn* AddedColumn = nullptr;
 		if (ObjectType->StructFlags & (STRUCT_IsPlainOldData | STRUCT_NoDestructor))
@@ -404,7 +416,7 @@ public:
 			FMassEntityHandle Entity;
 		};
 
-		FAddedColumns* AddedColumns = ScratchBuffer.Emplace<FAddedColumns>();
+		FAddedColumns* AddedColumns = Environment.GetScratchBuffer().Emplace<FAddedColumns>();
 		TedsColumnsToMassDescriptor(AddedColumns->AddDescriptor, ColumnTypes);
 		AddedColumns->Entity = FMassEntityHandle::FromNumber(Row);
 
@@ -424,6 +436,7 @@ public:
 			int32 EntityCount;
 		};
 
+		FTypedElementDatabaseScratchBuffer& ScratchBuffer = Environment.GetScratchBuffer();
 		FAddedColumns* AddedColumns = ScratchBuffer.Emplace<FAddedColumns>();
 		TedsColumnsToMassDescriptor(AddedColumns->AddDescriptor, ColumnTypes);
 		
@@ -456,7 +469,7 @@ public:
 			FMassEntityHandle Entity;
 		};
 
-		FRemovedColumns* RemovedColumns = ScratchBuffer.Emplace<FRemovedColumns>();
+		FRemovedColumns* RemovedColumns = Environment.GetScratchBuffer().Emplace<FRemovedColumns>();
 		TedsColumnsToMassDescriptor(RemovedColumns->RemoveDescriptor, ColumnTypes);
 		RemovedColumns->Entity = FMassEntityHandle::FromNumber(Row);
 
@@ -476,6 +489,7 @@ public:
 			int32 EntityCount;
 		};
 
+		FTypedElementDatabaseScratchBuffer& ScratchBuffer = Environment.GetScratchBuffer();
 		FRemovedColumns* RemovedColumns = ScratchBuffer.Emplace<FRemovedColumns>();
 		TedsColumnsToMassDescriptor(RemovedColumns->RemoveDescriptor, ColumnTypes);
 
@@ -500,6 +514,11 @@ public:
 			});
 	}
 
+	TypedElementDataStorage::RowHandle FindIndexedRow(TypedElementDataStorage::IndexHash Index) const override
+	{
+		return Environment.GetIndexTable().FindIndexedRow(Index);
+	}
+
 	TypedElementDataStorage::FQueryResult RunQuery(TypedElementQueryHandle Query) override
 	{
 		FTypedElementExtendedQueryStore::Handle Handle;
@@ -546,7 +565,7 @@ public:
 
 	ITypedElementDataStorageInterface::FQueryDescription& QueryDescription;
 	FTypedElementExtendedQueryStore& QueryStore;
-	FTypedElementDatabaseScratchBuffer& ScratchBuffer;
+	FTypedElementDatabaseEnvironment& Environment;
 };
 
 
@@ -570,16 +589,16 @@ FPhasePreOrPostAmbleExecutor::~FPhasePreOrPostAmbleExecutor()
 void FPhasePreOrPostAmbleExecutor::ExecuteQuery(
 	ITypedElementDataStorageInterface::FQueryDescription& Description, 
 	FTypedElementExtendedQueryStore& QueryStore, 
-	FTypedElementDatabaseScratchBuffer& ScratchBuffer,
+	FTypedElementDatabaseEnvironment& Environment,
 	FMassEntityQuery& NativeQuery,
 	ITypedElementDataStorageInterface::QueryCallbackRef Callback)
 {
 	NativeQuery.ForEachEntityChunk(Context.GetEntityManagerChecked(), Context,
-		[&Callback, &QueryStore, &ScratchBuffer, &Description](FMassExecutionContext& ExecutionContext)
+		[&Callback, &QueryStore, &Environment, &Description](FMassExecutionContext& ExecutionContext)
 		{
 			if (FTypedElementQueryProcessorData::PrepareCachedDependenciesOnQuery(Description, ExecutionContext))
 			{
-				FMassContextForwarder QueryContext(Description, ExecutionContext, QueryStore, ScratchBuffer);
+				FMassContextForwarder QueryContext(Description, ExecutionContext, QueryStore, Environment);
 				Callback(Description, QueryContext);
 			}
 		}
@@ -600,12 +619,12 @@ bool FTypedElementQueryProcessorData::CommonQueryConfiguration(
 	UMassProcessor& InOwner, 
 	FTypedElementExtendedQuery& InQuery, 
 	FTypedElementExtendedQueryStore& InQueryStore, 
-	FTypedElementDatabaseScratchBuffer& InScratchBuffer,
+	FTypedElementDatabaseEnvironment& InEnvironment,
 	TArrayView<FMassEntityQuery> Subqueries)
 {
 	ParentQuery = &InQuery;
 	QueryStore = &InQueryStore;
-	ScratchBuffer = &InScratchBuffer;
+	Environment = &InEnvironment;
 
 	if (ensureMsgf(InQuery.Description.Subqueries.Num() <= Subqueries.Num(),
 		TEXT("Provided query has too many (%i) subqueries."), InQuery.Description.Subqueries.Num()))
@@ -774,7 +793,7 @@ void FTypedElementQueryProcessorData::Execute(FMassEntityManager& EntityManager,
 		{
 			if (PrepareCachedDependenciesOnQuery(Description, Context))
 			{
-				FMassContextForwarder QueryContext(Description, Context, *QueryStore, *ScratchBuffer);
+				FMassContextForwarder QueryContext(Description, Context, *QueryStore, *Environment);
 				Description.Callback.Function(Description, QueryContext);
 			}
 		}
@@ -802,9 +821,9 @@ FMassEntityQuery& UTypedElementQueryProcessorCallbackAdapterProcessorBase::GetQu
 bool UTypedElementQueryProcessorCallbackAdapterProcessorBase::ConfigureQueryCallback(
 	FTypedElementExtendedQuery& Query, 
 	FTypedElementExtendedQueryStore& QueryStore, 
-	FTypedElementDatabaseScratchBuffer& ScratchBuffer)
+	FTypedElementDatabaseEnvironment& Environment)
 {
-	return ConfigureQueryCallbackData(Query, QueryStore, ScratchBuffer, {});
+	return ConfigureQueryCallbackData(Query, QueryStore, Environment, {});
 }
 
 bool UTypedElementQueryProcessorCallbackAdapterProcessorBase::ShouldAllowQueryBasedPruning(
@@ -819,10 +838,10 @@ bool UTypedElementQueryProcessorCallbackAdapterProcessorBase::ShouldAllowQueryBa
 bool UTypedElementQueryProcessorCallbackAdapterProcessorBase::ConfigureQueryCallbackData(
 	FTypedElementExtendedQuery& Query,
 	FTypedElementExtendedQueryStore& QueryStore, 
-	FTypedElementDatabaseScratchBuffer& ScratchBuffer,
+	FTypedElementDatabaseEnvironment& Environment,
 	TArrayView<FMassEntityQuery> Subqueries)
 {
-	bool Result = Data.CommonQueryConfiguration(*this, Query, QueryStore, ScratchBuffer, Subqueries);
+	bool Result = Data.CommonQueryConfiguration(*this, Query, QueryStore, Environment, Subqueries);
 
 	bRequiresGameThreadExecution = Query.Description.Callback.bForceToGameThread;
 	ExecutionFlags = static_cast<int32>(EProcessorExecutionFlags::Editor); 
@@ -858,27 +877,27 @@ void UTypedElementQueryProcessorCallbackAdapterProcessorBase::Execute(FMassEntit
 }
 
 bool UTypedElementQueryProcessorCallbackAdapterProcessorWith1Subquery::ConfigureQueryCallback(
-	FTypedElementExtendedQuery& Query, FTypedElementExtendedQueryStore& QueryStore, FTypedElementDatabaseScratchBuffer& ScratchBuffer)
+	FTypedElementExtendedQuery& Query, FTypedElementExtendedQueryStore& QueryStore, FTypedElementDatabaseEnvironment& Environment)
 {
-	return ConfigureQueryCallbackData(Query, QueryStore, ScratchBuffer, NativeSubqueries);
+	return ConfigureQueryCallbackData(Query, QueryStore, Environment, NativeSubqueries);
 }
 
 bool UTypedElementQueryProcessorCallbackAdapterProcessorWith2Subqueries::ConfigureQueryCallback(
-	FTypedElementExtendedQuery& Query, FTypedElementExtendedQueryStore& QueryStore, FTypedElementDatabaseScratchBuffer& ScratchBuffer)
+	FTypedElementExtendedQuery& Query, FTypedElementExtendedQueryStore& QueryStore, FTypedElementDatabaseEnvironment& Environment)
 {
-	return ConfigureQueryCallbackData(Query, QueryStore, ScratchBuffer, NativeSubqueries);
+	return ConfigureQueryCallbackData(Query, QueryStore, Environment, NativeSubqueries);
 }
 
 bool UTypedElementQueryProcessorCallbackAdapterProcessorWith3Subqueries::ConfigureQueryCallback(
-	FTypedElementExtendedQuery& Query, FTypedElementExtendedQueryStore& QueryStore, FTypedElementDatabaseScratchBuffer& ScratchBuffer)
+	FTypedElementExtendedQuery& Query, FTypedElementExtendedQueryStore& QueryStore, FTypedElementDatabaseEnvironment& Environment)
 {
-	return ConfigureQueryCallbackData(Query, QueryStore, ScratchBuffer, NativeSubqueries);
+	return ConfigureQueryCallbackData(Query, QueryStore, Environment, NativeSubqueries);
 }
 
 bool UTypedElementQueryProcessorCallbackAdapterProcessorWith4Subqueries::ConfigureQueryCallback(
-	FTypedElementExtendedQuery& Query, FTypedElementExtendedQueryStore& QueryStore, FTypedElementDatabaseScratchBuffer& ScratchBuffer)
+	FTypedElementExtendedQuery& Query, FTypedElementExtendedQueryStore& QueryStore, FTypedElementDatabaseEnvironment& Environment)
 {
-	return ConfigureQueryCallbackData(Query, QueryStore, ScratchBuffer, NativeSubqueries);
+	return ConfigureQueryCallbackData(Query, QueryStore, Environment, NativeSubqueries);
 }
 
 
@@ -909,15 +928,15 @@ EMassObservedOperation UTypedElementQueryObserverCallbackAdapterProcessorBase::G
 }
 
 bool UTypedElementQueryObserverCallbackAdapterProcessorBase::ConfigureQueryCallback(
-	FTypedElementExtendedQuery& Query, FTypedElementExtendedQueryStore& QueryStore, FTypedElementDatabaseScratchBuffer& ScratchBuffer)
+	FTypedElementExtendedQuery& Query, FTypedElementExtendedQueryStore& QueryStore, FTypedElementDatabaseEnvironment& Environment)
 {
-	return ConfigureQueryCallbackData(Query, QueryStore, ScratchBuffer, {});
+	return ConfigureQueryCallbackData(Query, QueryStore, Environment, {});
 }
 
 bool UTypedElementQueryObserverCallbackAdapterProcessorBase::ConfigureQueryCallbackData(FTypedElementExtendedQuery& Query,
-	FTypedElementExtendedQueryStore& QueryStore, FTypedElementDatabaseScratchBuffer& ScratchBuffer, TArrayView<FMassEntityQuery> Subqueries)
+	FTypedElementExtendedQueryStore& QueryStore, FTypedElementDatabaseEnvironment& Environment, TArrayView<FMassEntityQuery> Subqueries)
 {
-	bool Result = Data.CommonQueryConfiguration(*this, Query, QueryStore, ScratchBuffer, Subqueries);
+	bool Result = Data.CommonQueryConfiguration(*this, Query, QueryStore, Environment, Subqueries);
 
 	bRequiresGameThreadExecution = Query.Description.Callback.bForceToGameThread;
 	ExecutionFlags = static_cast<int32>(EProcessorExecutionFlags::Editor);
@@ -982,25 +1001,25 @@ void UTypedElementQueryObserverCallbackAdapterProcessorBase::Execute(FMassEntity
 }
 
 bool UTypedElementQueryObserverCallbackAdapterProcessorWith1Subquery::ConfigureQueryCallback(
-	FTypedElementExtendedQuery& Query, FTypedElementExtendedQueryStore& QueryStore, FTypedElementDatabaseScratchBuffer& ScratchBuffer)
+	FTypedElementExtendedQuery& Query, FTypedElementExtendedQueryStore& QueryStore, FTypedElementDatabaseEnvironment& Environment)
 {
-	return ConfigureQueryCallbackData(Query, QueryStore, ScratchBuffer, NativeSubqueries);
+	return ConfigureQueryCallbackData(Query, QueryStore, Environment, NativeSubqueries);
 }
 
 bool UTypedElementQueryObserverCallbackAdapterProcessorWith2Subqueries::ConfigureQueryCallback(
-	FTypedElementExtendedQuery& Query, FTypedElementExtendedQueryStore& QueryStore, FTypedElementDatabaseScratchBuffer& ScratchBuffer)
+	FTypedElementExtendedQuery& Query, FTypedElementExtendedQueryStore& QueryStore, FTypedElementDatabaseEnvironment& Environment)
 {
-	return ConfigureQueryCallbackData(Query, QueryStore, ScratchBuffer, NativeSubqueries);
+	return ConfigureQueryCallbackData(Query, QueryStore, Environment, NativeSubqueries);
 }
 
 bool UTypedElementQueryObserverCallbackAdapterProcessorWith3Subqueries::ConfigureQueryCallback(
-	FTypedElementExtendedQuery& Query, FTypedElementExtendedQueryStore& QueryStore, FTypedElementDatabaseScratchBuffer& ScratchBuffer)
+	FTypedElementExtendedQuery& Query, FTypedElementExtendedQueryStore& QueryStore, FTypedElementDatabaseEnvironment& Environment)
 {
-	return ConfigureQueryCallbackData(Query, QueryStore, ScratchBuffer, NativeSubqueries);
+	return ConfigureQueryCallbackData(Query, QueryStore, Environment, NativeSubqueries);
 }
 
 bool UTypedElementQueryObserverCallbackAdapterProcessorWith4Subqueries::ConfigureQueryCallback(
-	FTypedElementExtendedQuery& Query, FTypedElementExtendedQueryStore& QueryStore, FTypedElementDatabaseScratchBuffer& ScratchBuffer)
+	FTypedElementExtendedQuery& Query, FTypedElementExtendedQueryStore& QueryStore, FTypedElementDatabaseEnvironment& Environment)
 {
-	return ConfigureQueryCallbackData(Query, QueryStore, ScratchBuffer, NativeSubqueries);
+	return ConfigureQueryCallbackData(Query, QueryStore, Environment, NativeSubqueries);
 }
diff --git a/Engine/Plugins/Experimental/TypedElementsDataStorage/Source/TypedElementsDataStorage/Private/Processors/TypedElementProcessorAdaptors.h b/Engine/Plugins/Experimental/TypedElementsDataStorage/Source/TypedElementsDataStorage/Private/Processors/TypedElementProcessorAdaptors.h
--- a/Engine/Plugins/Experimental/TypedElementsDataStorage/Source/TypedElementsDataStorage/Private/Processors/TypedElementProcessorAdaptors.h
+++ b/Engine/Plugins/Experimental/TypedElementsDataStorage/Source/TypedElementsDataStorage/Private/Processors/TypedElementProcessorAdaptors.h
@@ -11,7 +11,7 @@
 
 #include "TypedElementProcessorAdaptors.generated.h"
 
-class FTypedElementDatabaseScratchBuffer;
+class FTypedElementDatabaseEnvironment;
 struct FTypedElementExtendedQuery;
 class FTypedElementExtendedQueryStore;
 
@@ -23,7 +23,7 @@ struct FPhasePreOrPostAmbleExecutor
 	void ExecuteQuery(
 		ITypedElementDataStorageInterface::FQueryDescription& Description,
 		FTypedElementExtendedQueryStore& QueryStore,
-		FTypedElementDatabaseScratchBuffer& ScratchBuffer,
+		FTypedElementDatabaseEnvironment& Environment,
 		FMassEntityQuery& NativeQuery,
 		ITypedElementDataStorageInterface::QueryCallbackRef Callback);
 
@@ -42,7 +42,7 @@ struct FTypedElementQueryProcessorData
 		UMassProcessor& InOwner, 
 		FTypedElementExtendedQuery& InQuery, 
 		FTypedElementExtendedQueryStore& InQueryStore, 
-		FTypedElementDatabaseScratchBuffer& InScratchBuffer,
+		FTypedElementDatabaseEnvironment& InEnvironment,
 		TArrayView<FMassEntityQuery> Subqueries);
 	static EMassProcessingPhase MapToMassProcessingPhase(ITypedElementDataStorageInterface::EQueryTickPhase Phase);
 	FString GetProcessorName() const;
@@ -71,7 +71,7 @@ struct FTypedElementQueryProcessorData
 	
 	FTypedElementExtendedQuery* ParentQuery{ nullptr };
 	FTypedElementExtendedQueryStore* QueryStore{ nullptr };
-	FTypedElementDatabaseScratchBuffer* ScratchBuffer{ nullptr };
+	FTypedElementDatabaseEnvironment* Environment{ nullptr };
 	FMassEntityQuery NativeQuery;
 };
 
@@ -90,7 +90,7 @@ public:
 	virtual bool ConfigureQueryCallback(
 		FTypedElementExtendedQuery& Query, 
 		FTypedElementExtendedQueryStore& QueryStore,
-		FTypedElementDatabaseScratchBuffer& ScratchBuffer);
+		FTypedElementDatabaseEnvironment& Environment);
 
 	virtual bool ShouldAllowQueryBasedPruning(const bool bRuntimeMode) const override;
 
@@ -98,7 +98,7 @@ protected:
 	bool ConfigureQueryCallbackData(
 		FTypedElementExtendedQuery& Query, 
 		FTypedElementExtendedQueryStore& QueryStore, 
-		FTypedElementDatabaseScratchBuffer& ScratchBuffer,
+		FTypedElementDatabaseEnvironment& Environment,
 		TArrayView<FMassEntityQuery> Subqueries);
 	void ConfigureQueries() override;
 	void Execute(FMassEntityManager& EntityManager, FMassExecutionContext& TargetParentQuery) override;
@@ -134,7 +134,7 @@ public:
 	bool ConfigureQueryCallback(
 		FTypedElementExtendedQuery& Query, 
 		FTypedElementExtendedQueryStore& QueryStore, 
-		FTypedElementDatabaseScratchBuffer& ScratchBuffer) override;
+		FTypedElementDatabaseEnvironment& Environment) override;
 
 private:
 	UPROPERTY(transient)
@@ -150,7 +150,7 @@ public:
 	bool ConfigureQueryCallback(
 		FTypedElementExtendedQuery& Query, 
 		FTypedElementExtendedQueryStore& QueryStore,
-		FTypedElementDatabaseScratchBuffer& ScratchBuffer) override;
+		FTypedElementDatabaseEnvironment& Environment) override;
 
 private:
 	UPROPERTY(transient)
@@ -166,7 +166,7 @@ public:
 	bool ConfigureQueryCallback(
 		FTypedElementExtendedQuery& Query, 
 		FTypedElementExtendedQueryStore& QueryStore,
-		FTypedElementDatabaseScratchBuffer& ScratchBuffer) override;
+		FTypedElementDatabaseEnvironment& Environment) override;
 
 private:
 	UPROPERTY(transient)
@@ -182,7 +182,7 @@ public:
 	bool ConfigureQueryCallback(
 		FTypedElementExtendedQuery& Query, 
 		FTypedElementExtendedQueryStore& QueryStore,
-		FTypedElementDatabaseScratchBuffer& ScratchBuffer) override;
+		FTypedElementDatabaseEnvironment& Environment) override;
 
 private:
 	UPROPERTY(transient)
@@ -204,11 +204,11 @@ public:
 	const UScriptStruct* GetObservedType() const;
 	EMassObservedOperation GetObservedOperation() const;
 	virtual bool ConfigureQueryCallback(
-		FTypedElementExtendedQuery& Query, FTypedElementExtendedQueryStore& QueryStore, FTypedElementDatabaseScratchBuffer& ScratchBuffer);
+		FTypedElementExtendedQuery& Query, FTypedElementExtendedQueryStore& QueryStore, FTypedElementDatabaseEnvironment& Environment);
 
 protected:
 	bool ConfigureQueryCallbackData(FTypedElementExtendedQuery& Query, FTypedElementExtendedQueryStore& QueryStore,
-		FTypedElementDatabaseScratchBuffer& ScratchBuffer, TArrayView<FMassEntityQuery> Subqueries);
+		FTypedElementDatabaseEnvironment& Environment, TArrayView<FMassEntityQuery> Subqueries);
 	void ConfigureQueries() override;
 	void Execute(FMassEntityManager& EntityManager, FMassExecutionContext& TargetParentQuery) override;
 
@@ -244,7 +244,7 @@ public:
 	bool ConfigureQueryCallback(
 		FTypedElementExtendedQuery& Query, 
 		FTypedElementExtendedQueryStore& QueryStore,
-		FTypedElementDatabaseScratchBuffer& ScratchBuffer) override;
+		FTypedElementDatabaseEnvironment& Environment) override;
 
 private:
 	UPROPERTY(transient)
@@ -260,7 +260,7 @@ public:
 	bool ConfigureQueryCallback(
 		FTypedElementExtendedQuery& Query, 
 		FTypedElementExtendedQueryStore& QueryStore,
-		FTypedElementDatabaseScratchBuffer& ScratchBuffer) override;
+		FTypedElementDatabaseEnvironment& Environment) override;
 
 private:
 	UPROPERTY(transient)
@@ -276,7 +276,7 @@ public:
 	bool ConfigureQueryCallback(
 		FTypedElementExtendedQuery& Query, 
 		FTypedElementExtendedQueryStore& QueryStore,
-		FTypedElementDatabaseScratchBuffer& ScratchBuffer) override;
+		FTypedElementDatabaseEnvironment& Environment) override;
 
 private:
 	UPROPERTY(transient)
@@ -292,7 +292,7 @@ public:
 	bool ConfigureQueryCallback(
 		FTypedElementExtendedQuery& Query, 
 		FTypedElementExtendedQueryStore& QueryStore,
-		FTypedElementDatabaseScratchBuffer& ScratchBuffer) override;
+		FTypedElementDatabaseEnvironment& Environment) override;
 
 private:
 	UPROPERTY(transient)
diff --git a/Engine/Plugins/Experimental/TypedElementsDataStorage/Source/TypedElementsDataStorage/Private/Queries/TypedElementExtendedQueryStore.cpp b/Engine/Plugins/Experimental/TypedElementsDataStorage/Source/TypedElementsDataStorage/Private/Queries/TypedElementExtendedQueryStore.cpp
--- a/Engine/Plugins/Experimental/TypedElementsDataStorage/Source/TypedElementsDataStorage/Private/Queries/TypedElementExtendedQueryStore.cpp
+++ b/Engine/Plugins/Experimental/TypedElementsDataStorage/Source/TypedElementsDataStorage/Private/Queries/TypedElementExtendedQueryStore.cpp
@@ -11,7 +11,7 @@ const ITypedElementDataStorageInterface::FQueryDescription FTypedElementExtended
 
 FTypedElementExtendedQueryStore::Handle FTypedElementExtendedQueryStore::RegisterQuery(
 	ITypedElementDataStorageInterface::FQueryDescription Query, 
-	FTypedElementDatabaseScratchBuffer& ScratchBuffer, 
+	FTypedElementDatabaseEnvironment& Environment,
 	FMassEntityManager& EntityManager, 
 	FMassProcessingPhaseManager& PhaseManager)
 {
@@ -24,7 +24,7 @@ FTypedElementExtendedQueryStore::Handle FTypedElementExtendedQueryStore::Registe
 	bContinueSetup = bContinueSetup &&	SetupConditions(StoredQuery.Description, NativeQuery);
 	bContinueSetup = bContinueSetup &&	SetupDependencies(StoredQuery.Description, NativeQuery);
 	bContinueSetup = bContinueSetup &&	SetupTickGroupDefaults(StoredQuery.Description);
-	bContinueSetup = bContinueSetup &&	SetupProcessors(Result, StoredQuery, ScratchBuffer, EntityManager, PhaseManager);
+	bContinueSetup = bContinueSetup &&	SetupProcessors(Result, StoredQuery, Environment, EntityManager, PhaseManager);
 	
 	if (!bContinueSetup)
 	{
@@ -272,16 +272,16 @@ TypedElementDataStorage::FQueryResult FTypedElementExtendedQueryStore::RunQuery(
 	return Result;
 }
 void FTypedElementExtendedQueryStore::RunPhasePreambleQueries(FMassEntityManager& EntityManager,
-	FTypedElementDatabaseScratchBuffer& ScratchBuffer, ITypedElementDataStorageInterface::EQueryTickPhase Phase, float DeltaTime)
+	FTypedElementDatabaseEnvironment& Environment, ITypedElementDataStorageInterface::EQueryTickPhase Phase, float DeltaTime)
 {
-	RunPhasePreOrPostAmbleQueries(EntityManager, ScratchBuffer, Phase, DeltaTime, 
+	RunPhasePreOrPostAmbleQueries(EntityManager, Environment, Phase, DeltaTime, 
 		PhasePreparationQueries[static_cast<QueryTickPhaseType>(Phase)]);
 }
 
 void FTypedElementExtendedQueryStore::RunPhasePostambleQueries(FMassEntityManager& EntityManager,
-	FTypedElementDatabaseScratchBuffer& ScratchBuffer, ITypedElementDataStorageInterface::EQueryTickPhase Phase, float DeltaTime)
+	FTypedElementDatabaseEnvironment& Environment, ITypedElementDataStorageInterface::EQueryTickPhase Phase, float DeltaTime)
 {
-	RunPhasePreOrPostAmbleQueries(EntityManager, ScratchBuffer, Phase, DeltaTime, 
+	RunPhasePreOrPostAmbleQueries(EntityManager, Environment, Phase, DeltaTime, 
 		PhaseFinalizationQueries[static_cast<QueryTickPhaseType>(Phase)]);
 }
 
@@ -570,7 +570,7 @@ bool FTypedElementExtendedQueryStore::SetupTickGroupDefaults(ITypedElementDataSt
 }
 
 bool FTypedElementExtendedQueryStore::SetupProcessors(Handle Query, FTypedElementExtendedQuery& StoredQuery, 
-	FTypedElementDatabaseScratchBuffer& ScratchBuffer, FMassEntityManager& EntityManager, FMassProcessingPhaseManager& PhaseManager)
+	FTypedElementDatabaseEnvironment& Environment, FMassEntityManager& EntityManager, FMassProcessingPhaseManager& PhaseManager)
 {
 	using DSI = ITypedElementDataStorageInterface;
 
@@ -591,7 +591,7 @@ bool FTypedElementExtendedQueryStore::SetupProcessors(Handle Query, FTypedElemen
 		if (StoredQuery.Processor->IsA<UTypedElementQueryProcessorCallbackAdapterProcessorBase>())
 		{
 			if (static_cast<UTypedElementQueryProcessorCallbackAdapterProcessorBase*>(StoredQuery.Processor.Get())->
-				ConfigureQueryCallback(StoredQuery, *this, ScratchBuffer))
+				ConfigureQueryCallback(StoredQuery, *this, Environment))
 			{
 				PhaseManager.RegisterDynamicProcessor(*StoredQuery.Processor);
 			}
@@ -605,7 +605,7 @@ bool FTypedElementExtendedQueryStore::SetupProcessors(Handle Query, FTypedElemen
 			if (UTypedElementQueryObserverCallbackAdapterProcessorBase* Observer =
 				static_cast<UTypedElementQueryObserverCallbackAdapterProcessorBase*>(StoredQuery.Processor.Get()))
 			{
-				Observer->ConfigureQueryCallback(StoredQuery, *this, ScratchBuffer);
+				Observer->ConfigureQueryCallback(StoredQuery, *this, Environment);
 				EntityManager.GetObserverManager().AddObserverInstance(*Observer->GetObservedType(), Observer->GetObservedOperation(), *Observer);
 			}
 			else
@@ -666,7 +666,7 @@ void FTypedElementExtendedQueryStore::UnregisterPostambleQuery(ITypedElementData
 }
 
 void FTypedElementExtendedQueryStore::RunPhasePreOrPostAmbleQueries(FMassEntityManager& EntityManager,
-	FTypedElementDatabaseScratchBuffer& ScratchBuffer, ITypedElementDataStorageInterface::EQueryTickPhase Phase, 
+	FTypedElementDatabaseEnvironment& Environment, ITypedElementDataStorageInterface::EQueryTickPhase Phase,
 	float DeltaTime, TArray<Handle>& QueryHandles)
 {
 	if (!QueryHandles.IsEmpty())
@@ -675,7 +675,7 @@ void FTypedElementExtendedQueryStore::RunPhasePreOrPostAmbleQueries(FMassEntityM
 		for (Handle Query : QueryHandles)
 		{
 			FTypedElementExtendedQuery& QueryData = Queries.Get(Query);
-			Executor.ExecuteQuery(QueryData.Description, *this, ScratchBuffer, QueryData.NativeQuery, QueryData.Description.Callback.Function);
+			Executor.ExecuteQuery(QueryData.Description, *this, Environment, QueryData.NativeQuery, QueryData.Description.Callback.Function);
 		}
 	}
 }
diff --git a/Engine/Plugins/Experimental/TypedElementsDataStorage/Source/TypedElementsDataStorage/Private/Queries/TypedElementExtendedQueryStore.h b/Engine/Plugins/Experimental/TypedElementsDataStorage/Source/TypedElementsDataStorage/Private/Queries/TypedElementExtendedQueryStore.h
--- a/Engine/Plugins/Experimental/TypedElementsDataStorage/Source/TypedElementsDataStorage/Private/Queries/TypedElementExtendedQueryStore.h
+++ b/Engine/Plugins/Experimental/TypedElementsDataStorage/Source/TypedElementsDataStorage/Private/Queries/TypedElementExtendedQueryStore.h
@@ -12,7 +12,7 @@ struct FMassEntityManager;
 struct FMassProcessingPhaseManager;
 class UMassProcessor;
 class FOutputDevice;
-class FTypedElementDatabaseScratchBuffer;
+class FTypedElementDatabaseEnvironment;
 
 struct FTypedElementExtendedQuery
 {
@@ -40,7 +40,7 @@ public:
 	/** Adds a new query to the store and initializes the query with the provided arguments. */
 	Handle RegisterQuery(
 		ITypedElementDataStorageInterface::FQueryDescription Query, 
-		FTypedElementDatabaseScratchBuffer& ScratchBuffer,
+		FTypedElementDatabaseEnvironment& Environment,
 		FMassEntityManager& EntityManager,
 		FMassProcessingPhaseManager& PhaseManager);
 	/** Removes the query at the given handle if still alive and otherwise does nothing. */
@@ -96,12 +96,12 @@ public:
 		TypedElementRowHandle Row, TypedElementDataStorage::SubqueryCallbackRef Callback);
 	void RunPhasePreambleQueries(
 		FMassEntityManager& EntityManager, 
-		FTypedElementDatabaseScratchBuffer& ScratchBuffer, 
+		FTypedElementDatabaseEnvironment& Environment,
 		ITypedElementDataStorageInterface::EQueryTickPhase Phase, 
 		float DeltaTime);
 	void RunPhasePostambleQueries(
 		FMassEntityManager& EntityManager, 
-		FTypedElementDatabaseScratchBuffer& ScratchBuffer, 
+		FTypedElementDatabaseEnvironment& Environment,
 		ITypedElementDataStorageInterface::EQueryTickPhase Phase, 
 		float DeltaTime);
 
@@ -136,7 +136,7 @@ private:
 	bool SetupConditions(ITypedElementDataStorageInterface::FQueryDescription& Query, FMassEntityQuery& NativeQuery);
 	bool SetupDependencies(ITypedElementDataStorageInterface::FQueryDescription& Query, FMassEntityQuery& NativeQuery);
 	bool SetupTickGroupDefaults(ITypedElementDataStorageInterface::FQueryDescription& Query);
-	bool SetupProcessors(Handle Query, FTypedElementExtendedQuery& StoredQuery, FTypedElementDatabaseScratchBuffer& ScratchBuffer,
+	bool SetupProcessors(Handle Query, FTypedElementExtendedQuery& StoredQuery, FTypedElementDatabaseEnvironment& Environment,
 		FMassEntityManager& EntityManager, FMassProcessingPhaseManager& PhaseManager);
 
 	EMassFragmentAccess ConvertToNativeAccessType(ITypedElementDataStorageInterface::EQueryAccessType AccessType);
@@ -145,7 +145,7 @@ private:
 	void RegisterPostambleQuery(ITypedElementDataStorageInterface::EQueryTickPhase Phase, Handle Query);
 	void UnregisterPreambleQuery(ITypedElementDataStorageInterface::EQueryTickPhase Phase, Handle Query);
 	void UnregisterPostambleQuery(ITypedElementDataStorageInterface::EQueryTickPhase Phase, Handle Query);
-	void RunPhasePreOrPostAmbleQueries(FMassEntityManager& EntityManager, FTypedElementDatabaseScratchBuffer& ScratchBuffer,
+	void RunPhasePreOrPostAmbleQueries(FMassEntityManager& EntityManager, FTypedElementDatabaseEnvironment& Environment,
 		ITypedElementDataStorageInterface::EQueryTickPhase Phase, float DeltaTime, TArray<Handle>& QueryHandles);
 
 	void UnregisterQueryData(Handle Query, FTypedElementExtendedQuery& QueryData, FMassProcessingPhaseManager& PhaseManager);
diff --git a/Engine/Plugins/Experimental/TypedElementsDataStorage/Source/TypedElementsDataStorage/Private/TypedElementDatabase.cpp b/Engine/Plugins/Experimental/TypedElementsDataStorage/Source/TypedElementsDataStorage/Private/TypedElementDatabase.cpp
--- a/Engine/Plugins/Experimental/TypedElementsDataStorage/Source/TypedElementsDataStorage/Private/TypedElementDatabase.cpp
+++ b/Engine/Plugins/Experimental/TypedElementsDataStorage/Source/TypedElementsDataStorage/Private/TypedElementDatabase.cpp
@@ -73,31 +73,35 @@ void UTypedElementDatabase::Initialize()
 
 	ActiveEditorEntityManager = Mass->GetMutableEntityManager();
 	ActiveEditorPhaseManager = Mass->GetMutablePhaseManager();
-
-	using PhaseType = std::underlying_type_t<EQueryTickPhase>;
-	for (PhaseType PhaseId = 0; PhaseId < static_cast<PhaseType>(EQueryTickPhase::Max); ++PhaseId)
+	if (ActiveEditorEntityManager && ActiveEditorPhaseManager)
 	{
-		EQueryTickPhase Phase = static_cast<EQueryTickPhase>(PhaseId);
-		EMassProcessingPhase MassPhase = FTypedElementQueryProcessorData::MapToMassProcessingPhase(Phase);
-		
-		ActiveEditorPhaseManager->GetOnPhaseStart(MassPhase).AddLambda(
-			[this, Phase](float DeltaTime)
-			{
-				PreparePhase(Phase, DeltaTime);
-			});
+		Environment = MakeUnique<FTypedElementDatabaseEnvironment>(*ActiveEditorEntityManager, *ActiveEditorPhaseManager);
 
-		ActiveEditorPhaseManager->GetOnPhaseEnd(MassPhase).AddLambda(
-			[this, Phase](float DeltaTime)
-			{
-				FinalizePhase(Phase, DeltaTime);
-			});
+		using PhaseType = std::underlying_type_t<EQueryTickPhase>;
+		for (PhaseType PhaseId = 0; PhaseId < static_cast<PhaseType>(EQueryTickPhase::Max); ++PhaseId)
+		{
+			EQueryTickPhase Phase = static_cast<EQueryTickPhase>(PhaseId);
+			EMassProcessingPhase MassPhase = FTypedElementQueryProcessorData::MapToMassProcessingPhase(Phase);
+
+			ActiveEditorPhaseManager->GetOnPhaseStart(MassPhase).AddLambda(
+				[this, Phase](float DeltaTime)
+				{
+					PreparePhase(Phase, DeltaTime);
+				});
 
-		// Guarantee that syncing to the data storage always happens before syncing to external.
-		RegisterTickGroup(GetQueryTickGroupName(EQueryTickGroups::SyncExternalToDataStorage), 
-			Phase, GetQueryTickGroupName(EQueryTickGroups::SyncDataStorageToExternal), {}, false);
-		// Guarantee that widgets syncs happen after external data has been updated to the data storage.
-		RegisterTickGroup(GetQueryTickGroupName(EQueryTickGroups::SyncWidgets),
-			Phase, {}, GetQueryTickGroupName(EQueryTickGroups::SyncExternalToDataStorage), false);
+			ActiveEditorPhaseManager->GetOnPhaseEnd(MassPhase).AddLambda(
+				[this, Phase](float DeltaTime)
+				{
+					FinalizePhase(Phase, DeltaTime);
+				});
+
+			// Guarantee that syncing to the data storage always happens before syncing to external.
+			RegisterTickGroup(GetQueryTickGroupName(EQueryTickGroups::SyncExternalToDataStorage),
+				Phase, GetQueryTickGroupName(EQueryTickGroups::SyncDataStorageToExternal), {}, false);
+			// Guarantee that widgets syncs happen after external data has been updated to the data storage.
+			RegisterTickGroup(GetQueryTickGroupName(EQueryTickGroups::SyncWidgets),
+				Phase, {}, GetQueryTickGroupName(EQueryTickGroups::SyncExternalToDataStorage), false);
+		}
 	}
 }
 
@@ -114,7 +118,7 @@ void UTypedElementDatabase::OnPreMassTick(float DeltaTime)
 	// to complete pending work.
 	ProcessPendingCommands();
 	// Recycle any full scratch memory blocks from the previous frame.
-	ScratchBuffer.RecycleBlocks();
+	Environment->GetScratchBuffer().RecycleBlocks();
 }
 
 TSharedPtr<FMassEntityManager> UTypedElementDatabase::GetActiveMutableEditorEntityManager()
@@ -252,6 +256,7 @@ void UTypedElementDatabase::RemoveRow(TypedElementRowHandle Row)
 	FMassEntityHandle Entity = FMassEntityHandle::FromNumber(Row);
 	if (ActiveEditorEntityManager && ActiveEditorEntityManager->IsEntityValid(Entity))
 	{
+		Environment->GetIndexTable().RemoveRow(Row);
 		if (ActiveEditorEntityManager->IsEntityBuilt(FMassEntityHandle::FromNumber(Row)))
 		{
 			ActiveEditorEntityManager->DestroyEntity(FMassEntityHandle::FromNumber(Row));
@@ -622,7 +627,7 @@ void UTypedElementDatabase::UnregisterTickGroup(FName GroupName, EQueryTickPhase
 TypedElementQueryHandle UTypedElementDatabase::RegisterQuery(FQueryDescription&& Query)
 {
 	return (ActiveEditorEntityManager && ActiveEditorPhaseManager)
-		? Queries.RegisterQuery(MoveTemp(Query), ScratchBuffer, *ActiveEditorEntityManager, *ActiveEditorPhaseManager).Handle
+		? Queries.RegisterQuery(MoveTemp(Query), *Environment, *ActiveEditorEntityManager, *ActiveEditorPhaseManager).Handle
 		: TypedElementInvalidQueryHandle;
 }
 
@@ -694,6 +699,27 @@ ITypedElementDataStorageInterface::FQueryResult UTypedElementDatabase::RunQuery(
 	}
 }
 
+TypedElementDataStorage::RowHandle UTypedElementDatabase::FindIndexedRow(TypedElementDataStorage::IndexHash Index) const
+{
+	return Environment->GetIndexTable().FindIndexedRow(Index);
+}
+
+void UTypedElementDatabase::IndexRow(TypedElementDataStorage::IndexHash Index, TypedElementDataStorage::RowHandle Row)
+{
+	Environment->GetIndexTable().IndexRow(Index, Row);
+}
+
+void UTypedElementDatabase::ReindexRow(TypedElementDataStorage::IndexHash OriginalIndex, TypedElementDataStorage::IndexHash NewIndex, 
+	TypedElementDataStorage::RowHandle RowHandle)
+{
+	Environment->GetIndexTable().ReindexRow(OriginalIndex, NewIndex, RowHandle);
+}
+
+void UTypedElementDatabase::RemoveIndex(TypedElementDataStorage::IndexHash Index)
+{
+	Environment->GetIndexTable().RemoveIndex(Index);
+}
+
 FTypedElementOnDataStorageUpdate& UTypedElementDatabase::OnUpdate()
 {
 	return OnUpdateDelegate;
@@ -827,7 +853,7 @@ void UTypedElementDatabase::PreparePhase(EQueryTickPhase Phase, float DeltaTime)
 {
 	if (ActiveEditorEntityManager)
 	{
-		Queries.RunPhasePreambleQueries(*ActiveEditorEntityManager, ScratchBuffer, Phase, DeltaTime);
+		Queries.RunPhasePreambleQueries(*ActiveEditorEntityManager, *Environment, Phase, DeltaTime);
 	}
 }
 
@@ -835,7 +861,7 @@ void UTypedElementDatabase::FinalizePhase(EQueryTickPhase Phase, float DeltaTime
 {
 	if (ActiveEditorEntityManager)
 	{
-		Queries.RunPhasePostambleQueries(*ActiveEditorEntityManager, ScratchBuffer, Phase, DeltaTime);
+		Queries.RunPhasePostambleQueries(*ActiveEditorEntityManager, *Environment, Phase, DeltaTime);
 	}
 }
 
@@ -847,6 +873,7 @@ void UTypedElementDatabase::Reset()
 	}
 	Tables.Reset();
 	TableNameLookup.Reset();
+	Environment.Reset();
 	ActiveEditorPhaseManager.Reset();
 	ActiveEditorEntityManager.Reset();
 }
diff --git a/Engine/Plugins/Experimental/TypedElementsDataStorage/Source/TypedElementsDataStorage/Private/TypedElementDatabase.h b/Engine/Plugins/Experimental/TypedElementsDataStorage/Source/TypedElementsDataStorage/Private/TypedElementDatabase.h
--- a/Engine/Plugins/Experimental/TypedElementsDataStorage/Source/TypedElementsDataStorage/Private/TypedElementDatabase.h
+++ b/Engine/Plugins/Experimental/TypedElementsDataStorage/Source/TypedElementsDataStorage/Private/TypedElementDatabase.h
@@ -2,12 +2,13 @@
 
 #pragma once
 
-#include "Templates/SharedPointer.h"
+#include "Containers/Map.h"
 #include "Elements/Interfaces/TypedElementDataStorageInterface.h"
 #include "MassArchetypeTypes.h"
 #include "Misc/TVariant.h"
 #include "Queries/TypedElementExtendedQueryStore.h"
-#include "TypedElementDatabaseScratchBuffer.h"
+#include "Templates/SharedPointer.h"
+#include "TypedElementDatabaseEnvironment.h"
 #include "UObject/ObjectMacros.h"
 #include "UObject/StrongObjectPtr.h"
 
@@ -81,6 +82,14 @@ public:
 	FQueryResult RunQuery(TypedElementQueryHandle Query) override;
 	FQueryResult RunQuery(TypedElementQueryHandle Query, DirectQueryCallbackRef Callback) override;
 
+	TypedElementDataStorage::RowHandle FindIndexedRow(TypedElementDataStorage::IndexHash Index) const override;
+	void IndexRow(TypedElementDataStorage::IndexHash Index, TypedElementDataStorage::RowHandle Row) override;
+	void ReindexRow(
+		TypedElementDataStorage::IndexHash OriginalIndex, 
+		TypedElementDataStorage::IndexHash NewIndex, 
+		TypedElementDataStorage::RowHandle Row) override;
+	void RemoveIndex(TypedElementDataStorage::IndexHash Index) override;
+
 	FTypedElementOnDataStorageUpdate& OnUpdate() override;
 	bool IsAvailable() const override;
 	void* GetExternalSystemAddress(UClass* Target) override;
@@ -134,9 +143,9 @@ private:
 	TArray<FMassArchetypeHandle> Tables;
 	TMap<FName, TypedElementTableHandle> TableNameLookup;
 
+	TUniquePtr<FTypedElementDatabaseEnvironment> Environment;
 	FTypedElementExtendedQueryStore Queries;
-	FTypedElementDatabaseScratchBuffer ScratchBuffer;
-
+	
 	FTypedElementOnDataStorageUpdate OnUpdateDelegate;
 
 	TSharedPtr<FMassEntityManager> ActiveEditorEntityManager;
diff --git a/Engine/Plugins/Experimental/TypedElementsDataStorage/Source/TypedElementsDataStorage/Private/TypedElementDatabaseCompatibility.cpp b/Engine/Plugins/Experimental/TypedElementsDataStorage/Source/TypedElementsDataStorage/Private/TypedElementDatabaseCompatibility.cpp
--- a/Engine/Plugins/Experimental/TypedElementsDataStorage/Source/TypedElementsDataStorage/Private/TypedElementDatabaseCompatibility.cpp
+++ b/Engine/Plugins/Experimental/TypedElementsDataStorage/Source/TypedElementsDataStorage/Private/TypedElementDatabaseCompatibility.cpp
@@ -10,6 +10,7 @@
 #include "Elements/Columns/TypedElementPackageColumns.h"
 #include "Elements/Columns/TypedElementTransformColumns.h"
 #include "Elements/Columns/TypedElementTypeInfoColumns.h"
+#include "Elements/Framework/TypedElementIndexHasher.h"
 #include "MassActorSubsystem.h"
 #include "TypedElementDataStorageProfilingMacros.h"
 
@@ -85,6 +86,8 @@ TypedElementRowHandle UTypedElementDatabaseCompatibility::AddCompatibleObjectExp
 
 TypedElementRowHandle UTypedElementDatabaseCompatibility::AddCompatibleObjectExplicit(UObject* Object, TypedElementTableHandle Table)
 {
+	using namespace TypedElementDataStorage;
+
 #if TEDS_SEPARATE_ACTOR_REGISTRATION
 	if (Object->IsA<AActor>())
 	{
@@ -97,7 +100,7 @@ TypedElementRowHandle UTypedElementDatabaseCompatibility::AddCompatibleObjectExp
 			ShouldAddObject(Object))
 		{
 			TypedElementRowHandle ReservedRow = Storage->ReserveRow();
-			ReverseObjectLookup.Add(Object, ReservedRow);
+			Storage->IndexRow(GenerateIndexHash(Object), ReservedRow);
 			
 			PendingRegistration<TWeakObjectPtr<UObject>>& Pending = UObjectsPendingRegistration.FindOrAdd(Table);
 			Pending.Add(ReservedRow, Object);
@@ -145,10 +148,12 @@ TypedElementRowHandle UTypedElementDatabaseCompatibility::AddCompatibleObjectExp
 TypedElementRowHandle UTypedElementDatabaseCompatibility::AddCompatibleObjectExplicit(
 	void* Object, TWeakObjectPtr<const UScriptStruct> TypeInfo, TypedElementTableHandle Table)
 {
+	using namespace TypedElementDataStorage;
+
 	if (ensureMsgf(Storage, TEXT("Trying to add an object to Typed Element's Data Storage before the storage is available.")))
 	{
 		TypedElementRowHandle ReservedRow = Storage->ReserveRow();
-		ReverseObjectLookup.Add(Object, ReservedRow);
+		Storage->IndexRow(GenerateIndexHash(Object), ReservedRow);
 		PendingRegistration<ExternalObjectRegistration>& Pending = ExternalObjectsPendingRegistration.FindOrAdd(Table);
 		Pending.Add(ReservedRow, ExternalObjectRegistration{ .Object = Object, .TypeInfo = TypeInfo });
 		return ReservedRow;
@@ -161,6 +166,8 @@ TypedElementRowHandle UTypedElementDatabaseCompatibility::AddCompatibleObjectExp
 
 void UTypedElementDatabaseCompatibility::RemoveCompatibleObjectExplicit(UObject* Object)
 {
+	using namespace TypedElementDataStorage;
+
 #if TEDS_SEPARATE_ACTOR_REGISTRATION
 	if (Object->IsA<AActor>())
 	{
@@ -170,8 +177,9 @@ void UTypedElementDatabaseCompatibility::RemoveCompatibleObjectExplicit(UObject*
 #else
 	{
 		checkf(Storage, TEXT("Removing compatible objects is not supported before Typed Element's Database compatibility manager has been initialized."));
-		TypedElementRowHandle Row;
-		if (ReverseObjectLookup.RemoveAndCopyValue(Object, Row))
+		IndexHash Hash = GenerateIndexHash(Object);
+		RowHandle Row = Storage->FindIndexedRow(Hash);
+		if (Storage->IsRowAvailable(Row))
 		{
 			const FTypedElementClassTypeInfoColumn* TypeInfoColumn = Storage->GetColumn<FTypedElementClassTypeInfoColumn>(Row);
 			if (Storage->HasRowBeenAssigned(Row) && ensureMsgf(TypeInfoColumn, TEXT("Missing type information for removed UObject at ptr 0x%p [%s]"), Object, *Object->GetName()))
@@ -186,9 +194,12 @@ void UTypedElementDatabaseCompatibility::RemoveCompatibleObjectExplicit(UObject*
 
 void UTypedElementDatabaseCompatibility::RemoveCompatibleObjectExplicit(void* Object)
 {
+	using namespace TypedElementDataStorage;
+
 	checkf(Storage, TEXT("Removing compatible objects is not supported before Typed Element's Database compatibility manager has been initialized."));
-	TypedElementRowHandle Row;
-	if (ReverseObjectLookup.RemoveAndCopyValue(Object, Row))
+	IndexHash Hash = GenerateIndexHash(Object);
+	RowHandle Row = Storage->FindIndexedRow(Hash);
+	if (Storage->IsRowAvailable(Row))
 	{
 		const FTypedElementScriptStructTypeInfoColumn* TypeInfoColumn = Storage->GetColumn<FTypedElementScriptStructTypeInfoColumn>(Row);
 		if (Storage->HasRowBeenAssigned(Row) && ensureMsgf(TypeInfoColumn, TEXT("Missing type information for removed void* object at ptr 0x%p"), Object))
@@ -238,6 +249,8 @@ void UTypedElementDatabaseCompatibility::RemoveCompatibleObjectExplicit(AActor*
 
 TypedElementRowHandle UTypedElementDatabaseCompatibility::FindRowWithCompatibleObjectExplicit(const UObject* Object) const
 {
+	using namespace TypedElementDataStorage;
+
 	if (Object && Storage && Storage->IsAvailable())
 	{
 #if TEDS_SEPARATE_ACTOR_REGISTRATION
@@ -250,8 +263,8 @@ TypedElementRowHandle UTypedElementDatabaseCompatibility::FindRowWithCompatibleO
 		else
 #endif
 		{
-			const TypedElementRowHandle* Row = ReverseObjectLookup.Find(Object);
-			return Row ? *Row : DealiasObject(Object);
+			RowHandle Row = Storage->FindIndexedRow(GenerateIndexHash(Object));
+			return Storage->IsRowAvailable(Row) ? Row : DealiasObject(Object);
 		}
 
 
@@ -278,12 +291,9 @@ TypedElementRowHandle UTypedElementDatabaseCompatibility::FindRowWithCompatibleO
 
 TypedElementRowHandle UTypedElementDatabaseCompatibility::FindRowWithCompatibleObjectExplicit(const void* Object) const
 {
-	if (Object && Storage && Storage->IsAvailable())
-	{
-		const TypedElementRowHandle* Row = ReverseObjectLookup.Find(Object);
-		return Row ? *Row : TypedElementInvalidRowHandle;
-	}
-	return TypedElementInvalidRowHandle;
+	using namespace TypedElementDataStorage;
+
+	return (Object && Storage && Storage->IsAvailable()) ? Storage->FindIndexedRow(GenerateIndexHash(Object)) : InvalidRowHandle;
 }
 
 void UTypedElementDatabaseCompatibility::Prepare()
diff --git a/Engine/Plugins/Experimental/TypedElementsDataStorage/Source/TypedElementsDataStorage/Private/TypedElementDatabaseCompatibility.h b/Engine/Plugins/Experimental/TypedElementsDataStorage/Source/TypedElementsDataStorage/Private/TypedElementDatabaseCompatibility.h
--- a/Engine/Plugins/Experimental/TypedElementsDataStorage/Source/TypedElementsDataStorage/Private/TypedElementDatabaseCompatibility.h
+++ b/Engine/Plugins/Experimental/TypedElementsDataStorage/Source/TypedElementsDataStorage/Private/TypedElementDatabaseCompatibility.h
@@ -3,6 +3,7 @@
 #pragma once
 
 #include "Containers/Array.h"
+#include "Containers/Map.h"
 #include "Elements/Interfaces/TypedElementDataStorageInterface.h"
 #include "Elements/Interfaces/TypedElementDataStorageCompatibilityInterface.h"
 #include "Compatibility/TypedElementObjectReinstancingManager.h"
@@ -124,8 +125,6 @@ private:
 	TSharedPtr<FMassActorManager> ActorSubsystem;
 #endif
 
-	TMap<void*, TypedElementRowHandle> ReverseObjectLookup;
-
 	/**
 	 * Reference of objects (UObject and AActor) that need to be fully synced from the world to the database.
 	 * Caution: Could point to objects that have been GC-ed
diff --git a/Engine/Plugins/Experimental/TypedElementsDataStorage/Source/TypedElementsDataStorage/Private/TypedElementDatabaseEnvironment.cpp b/Engine/Plugins/Experimental/TypedElementsDataStorage/Source/TypedElementsDataStorage/Private/TypedElementDatabaseEnvironment.cpp
new file mode 100644
--- /dev/null
+++ b/Engine/Plugins/Experimental/TypedElementsDataStorage/Source/TypedElementsDataStorage/Private/TypedElementDatabaseEnvironment.cpp
@@ -0,0 +1,51 @@
+// Copyright Epic Games, Inc. All Rights Reserved.
+
+#include "TypedElementDatabaseEnvironment.h"
+
+FTypedElementDatabaseEnvironment::FTypedElementDatabaseEnvironment(
+	FMassEntityManager& InMassEntityManager, FMassProcessingPhaseManager& InMassPhaseManager)
+	: MassEntityManager(InMassEntityManager)
+	, MassPhaseManager(InMassPhaseManager)
+{
+
+}
+
+FTypedElementDatabaseIndexTable& FTypedElementDatabaseEnvironment::GetIndexTable()
+{
+	return IndexTable;
+}
+
+const FTypedElementDatabaseIndexTable& FTypedElementDatabaseEnvironment::GetIndexTable() const
+{
+	return IndexTable;
+}
+
+FTypedElementDatabaseScratchBuffer& FTypedElementDatabaseEnvironment::GetScratchBuffer()
+{
+	return ScratchBuffer;
+}
+
+const FTypedElementDatabaseScratchBuffer& FTypedElementDatabaseEnvironment::GetScratchBuffer() const
+{
+	return ScratchBuffer;
+}
+
+FMassEntityManager& FTypedElementDatabaseEnvironment::GetMassEntityManager()
+{
+	return MassEntityManager;
+}
+
+const FMassEntityManager& FTypedElementDatabaseEnvironment::GetMassEntityManager() const
+{
+	return MassEntityManager;
+}
+
+FMassProcessingPhaseManager& FTypedElementDatabaseEnvironment::GetMassPhaseManager()
+{
+	return MassPhaseManager;
+}
+
+const FMassProcessingPhaseManager& FTypedElementDatabaseEnvironment::GetMassPhaseManager() const
+{
+	return MassPhaseManager;
+}
diff --git a/Engine/Plugins/Experimental/TypedElementsDataStorage/Source/TypedElementsDataStorage/Private/TypedElementDatabaseEnvironment.h b/Engine/Plugins/Experimental/TypedElementsDataStorage/Source/TypedElementsDataStorage/Private/TypedElementDatabaseEnvironment.h
new file mode 100644
--- /dev/null
+++ b/Engine/Plugins/Experimental/TypedElementsDataStorage/Source/TypedElementsDataStorage/Private/TypedElementDatabaseEnvironment.h
@@ -0,0 +1,33 @@
+// Copyright Epic Games, Inc. All Rights Reserved.
+
+#pragma once
+
+#include "MassEntityManager.h"
+#include "MassProcessingPhaseManager.h"
+#include "TypedElementDatabaseScratchBuffer.h"
+#include "TypedElementDatabaseIndexTable.h"
+
+class FTypedElementDatabaseEnvironment final
+{
+public:
+	FTypedElementDatabaseEnvironment(FMassEntityManager& InMassEntityManager, FMassProcessingPhaseManager& InMassPhaseManager);
+
+	FTypedElementDatabaseIndexTable& GetIndexTable();
+	const FTypedElementDatabaseIndexTable& GetIndexTable() const;
+
+	FTypedElementDatabaseScratchBuffer& GetScratchBuffer();
+	const FTypedElementDatabaseScratchBuffer& GetScratchBuffer() const;
+
+	FMassEntityManager& GetMassEntityManager();
+	const FMassEntityManager& GetMassEntityManager() const;
+	
+	FMassProcessingPhaseManager& GetMassPhaseManager();
+	const FMassProcessingPhaseManager& GetMassPhaseManager() const;
+
+private:
+	FTypedElementDatabaseIndexTable IndexTable;
+	FTypedElementDatabaseScratchBuffer ScratchBuffer;
+
+	FMassEntityManager& MassEntityManager;
+	FMassProcessingPhaseManager& MassPhaseManager;
+};
diff --git a/Engine/Plugins/Experimental/TypedElementsDataStorage/Source/TypedElementsDataStorage/Private/TypedElementDatabaseIndexTable.cpp b/Engine/Plugins/Experimental/TypedElementsDataStorage/Source/TypedElementsDataStorage/Private/TypedElementDatabaseIndexTable.cpp
new file mode 100644
--- /dev/null
+++ b/Engine/Plugins/Experimental/TypedElementsDataStorage/Source/TypedElementsDataStorage/Private/TypedElementDatabaseIndexTable.cpp
@@ -0,0 +1,48 @@
+// Copyright Epic Games, Inc. All Rights Reserved.
+
+#include "TypedElementDatabaseIndexTable.h"
+
+TypedElementDataStorage::RowHandle FTypedElementDatabaseIndexTable::FindIndexedRow(TypedElementDataStorage::IndexHash Index) const
+{
+	const TypedElementDataStorage::RowHandle* Result = IndexLookupMap.Find(Index);
+	return Result ? *Result : TypedElementDataStorage::InvalidRowHandle;
+}
+
+void FTypedElementDatabaseIndexTable::IndexRow(TypedElementDataStorage::IndexHash Index, TypedElementDataStorage::RowHandle Row)
+{
+	IndexLookupMap.Add(Index, Row);
+	ReverseIndexLookupMap.Add(Row, Index);
+}
+
+void FTypedElementDatabaseIndexTable::ReindexRow(TypedElementDataStorage::IndexHash OriginalIndex, 
+	TypedElementDataStorage::IndexHash NewIndex, TypedElementDataStorage::RowHandle Row)
+{
+	RemoveIndex(OriginalIndex);
+	IndexRow(NewIndex, Row);
+}
+
+void FTypedElementDatabaseIndexTable::RemoveIndex(TypedElementDataStorage::IndexHash Index)
+{
+	using namespace TypedElementDataStorage;
+
+	if (const RowHandle* Row = IndexLookupMap.Find(Index))
+	{
+		IndexLookupMap.Remove(Index);
+		ReverseIndexLookupMap.Remove(*Row, Index);
+	}
+}
+
+void FTypedElementDatabaseIndexTable::RemoveRow(TypedElementDataStorage::RowHandle Row)
+{
+	using namespace TypedElementDataStorage;
+	
+	if (TMultiMap<RowHandle, IndexHash>::TKeyIterator It = ReverseIndexLookupMap.CreateKeyIterator(Row); It)
+	{
+		do
+		{
+			IndexLookupMap.Remove(It.Value());
+			++It;
+		} while (It);
+		ReverseIndexLookupMap.Remove(Row);
+	}
+}
\ No newline at end of file
diff --git a/Engine/Plugins/Experimental/TypedElementsDataStorage/Source/TypedElementsDataStorage/Private/TypedElementDatabaseIndexTable.h b/Engine/Plugins/Experimental/TypedElementsDataStorage/Source/TypedElementsDataStorage/Private/TypedElementDatabaseIndexTable.h
new file mode 100644
--- /dev/null
+++ b/Engine/Plugins/Experimental/TypedElementsDataStorage/Source/TypedElementsDataStorage/Private/TypedElementDatabaseIndexTable.h
@@ -0,0 +1,24 @@
+// Copyright Epic Games, Inc. All Rights Reserved.
+
+#pragma once
+
+#include "Containers/Map.h"
+#include "Elements/Common/TypedElementHandles.h"
+#include "Elements/Interfaces/TypedElementQueryStorageInterfaces.h"
+
+class FTypedElementDatabaseIndexTable final
+{
+public:
+	TypedElementDataStorage::RowHandle FindIndexedRow(TypedElementDataStorage::IndexHash Index) const;
+	void IndexRow(TypedElementDataStorage::IndexHash Index, TypedElementDataStorage::RowHandle Row);
+	void ReindexRow(
+		TypedElementDataStorage::IndexHash OriginalIndex,
+		TypedElementDataStorage::IndexHash NewIndex,
+		TypedElementDataStorage::RowHandle Row);
+	void RemoveIndex(TypedElementDataStorage::IndexHash Index);
+	void RemoveRow(TypedElementDataStorage::RowHandle Row);
+
+private:
+	TMap<TypedElementDataStorage::IndexHash, TypedElementDataStorage::RowHandle> IndexLookupMap;
+	TMultiMap<TypedElementDataStorage::RowHandle, TypedElementDataStorage::IndexHash> ReverseIndexLookupMap;
+};
diff --git a/Engine/Source/Runtime/TypedElementFramework/Public/Elements/Framework/TypedElementIndexHasher.h b/Engine/Source/Runtime/TypedElementFramework/Public/Elements/Framework/TypedElementIndexHasher.h
new file mode 100644
--- /dev/null
+++ b/Engine/Source/Runtime/TypedElementFramework/Public/Elements/Framework/TypedElementIndexHasher.h
@@ -0,0 +1,15 @@
+// Copyright Epic Games, Inc. All Rights Reserved.
+
+#pragma once
+
+#include "Elements/Interfaces/TypedElementQueryStorageInterfaces.h"
+
+namespace TypedElementDataStorage
+{
+	template<typename T>
+	IndexHash GenerateIndexHash(const T* Object);
+	template<typename T>
+	IndexHash GenerateIndexHash(T* Object);
+} // namespace TypedElementDataStorage
+
+#include "Elements/Framework/TypedElementIndexHasher.inl"
\ No newline at end of file
diff --git a/Engine/Source/Runtime/TypedElementFramework/Public/Elements/Framework/TypedElementIndexHasher.inl b/Engine/Source/Runtime/TypedElementFramework/Public/Elements/Framework/TypedElementIndexHasher.inl
new file mode 100644
--- /dev/null
+++ b/Engine/Source/Runtime/TypedElementFramework/Public/Elements/Framework/TypedElementIndexHasher.inl
@@ -0,0 +1,28 @@
+// Copyright Epic Games, Inc. All Rights Reserved.
+
+#pragma once
+
+#include <type_traits>
+#include "UObject/Object.h"
+
+namespace TypedElementDataStorage
+{
+	template<typename T>
+	IndexHash GenerateIndexHash(const T* Object)
+	{
+		if constexpr (std::is_base_of_v<UObject, T>)
+		{
+			return Object->GetUniqueID();
+		}
+		else
+		{
+			return reinterpret_cast<IndexHash>(Object);
+		}
+	}
+
+	template<typename T>
+	IndexHash GenerateIndexHash(T* Object)
+	{
+		return GenerateIndexHash(const_cast<const T*>(Object));
+	}
+} // namespace TypedElementDataStorage
diff --git a/Engine/Source/Runtime/TypedElementFramework/Public/Elements/Interfaces/TypedElementDataStorageInterface.h b/Engine/Source/Runtime/TypedElementFramework/Public/Elements/Interfaces/TypedElementDataStorageInterface.h
--- a/Engine/Source/Runtime/TypedElementFramework/Public/Elements/Interfaces/TypedElementDataStorageInterface.h
+++ b/Engine/Source/Runtime/TypedElementFramework/Public/Elements/Interfaces/TypedElementDataStorageInterface.h
@@ -256,6 +256,26 @@ public:
 	 */
 	virtual FQueryResult RunQuery(TypedElementQueryHandle Query, DirectQueryCallbackRef Callback) = 0;
 	
+	/**
+	 * @section Indexing
+	 * @description
+	 * In order for rows to reference each other it's often needed to find a row based on the content of one of its columns. This can be
+	 * done by linearly searching through columns, though this comes at a performance cost. As an alternative the data storage allows
+	 * one or more indexes to be created for a row. An index is a 64-bit value and typically uses a hash value of an identifying value.
+	 */
+
+	/** Retrieves the row for an indexed object. Returns an invalid row handle if the hash wasn't found. */
+	virtual TypedElementDataStorage::RowHandle FindIndexedRow(TypedElementDataStorage::IndexHash Index) const = 0;
+	/** 
+	 * Registers a row under the index hash. The same row can be registered multiple, but an index hash can only be associated 
+	 * with a single row.
+	 */
+	virtual void IndexRow(TypedElementDataStorage::IndexHash Index, TypedElementDataStorage::RowHandle Row) = 0;
+	/** Updates the index of a row to a new value. Effectively this is the same as removing an index and adding a new one. */
+	virtual void ReindexRow(TypedElementDataStorage::IndexHash OriginalIndex, TypedElementDataStorage::IndexHash NewIndex, TypedElementDataStorage::RowHandle Row) = 0;
+	/** Removes a previously registered index hash from the index lookup table or does nothing if the hash no longer exists. */
+	virtual void RemoveIndex(TypedElementDataStorage::IndexHash Index) = 0;
+
 	/**
 	 * @section Misc
 	 */
diff --git a/Engine/Source/Runtime/TypedElementFramework/Public/Elements/Interfaces/TypedElementQueryStorageInterfaces.h b/Engine/Source/Runtime/TypedElementFramework/Public/Elements/Interfaces/TypedElementQueryStorageInterfaces.h
--- a/Engine/Source/Runtime/TypedElementFramework/Public/Elements/Interfaces/TypedElementQueryStorageInterfaces.h
+++ b/Engine/Source/Runtime/TypedElementFramework/Public/Elements/Interfaces/TypedElementQueryStorageInterfaces.h
@@ -21,6 +21,8 @@ namespace TypedElementDataStorage
 	using SubqueryCallback = TFunction<void(const FQueryDescription&, ISubqueryContext&)>;
 	using SubqueryCallbackRef = TFunctionRef<void(const FQueryDescription&, ISubqueryContext&)>;
 
+	using IndexHash = uint64;
+
 	/**
 	 * Base interface for any contexts provided to query callbacks.
 	 */
@@ -28,6 +30,14 @@ namespace TypedElementDataStorage
 	{
 		virtual ~ICommonQueryContext() = default;
 
+		/** Returns the number rows in the batch. */
+		virtual uint32 GetRowCount() const = 0;
+		/**
+		 * Returns an immutable view that contains the row handles for all returned results. The returned size will be the same  as the
+		 * value returned by GetRowCount().
+		 */
+		virtual TConstArrayView<RowHandle> GetRowHandles() const = 0;
+
 		/** Return the address of a immutable column matching the requested type or a nullptr if not found. */
 		virtual const void* GetColumn(const UScriptStruct* ColumnType) const = 0;
 		/** Return the address of a mutable column matching the requested type or a nullptr if not found. */
@@ -46,14 +56,6 @@ namespace TypedElementDataStorage
 		virtual void GetColumnsUnguarded(int32 TypeCount, char** RetrievedAddresses, const TWeakObjectPtr<const UScriptStruct>* ColumnTypes,
 			const EQueryAccessType* AccessTypes) = 0;
 
-		/** Returns the number rows in the batch. */
-		virtual uint32 GetRowCount() const = 0;
-		/**
-		 * Returns an immutable view that contains the row handles for all returned results. The returned size will be the same  as the
-		 * value returned by GetRowCount().
-		 */
-		virtual TConstArrayView<RowHandle> GetRowHandles() const = 0;
-
 		// Utility functions
 
 		template<typename Column>
@@ -99,6 +101,10 @@ namespace TypedElementDataStorage
 		virtual void GetDependencies(TArrayView<UObject*> RetrievedAddresses, TConstArrayView<TWeakObjectPtr<const UClass>> DependencyTypes,
 			TConstArrayView<EQueryAccessType> AccessTypes) = 0;
 
+		/** Checks whether or not a row is in use. This is true even if the row has only been reserved. */
+		virtual bool IsRowAvailable(RowHandle Row) const = 0;
+		/** Checks whether or not a row has been reserved but not yet assigned to a table. */
+		virtual bool HasRowBeenAssigned(RowHandle Row) const = 0;
 		/**
 		 * Removes the row with the provided row handle. The removal will not be immediately done but delayed until the end of the tick
 		 * group.
@@ -150,6 +156,9 @@ namespace TypedElementDataStorage
 		 */
 		virtual void RemoveColumns(TConstArrayView<RowHandle> Rows, TConstArrayView<const UScriptStruct*> ColumnTypes) = 0;
 
+		/** Retrieves the row for an indexed object. Returns an invalid row handle if the hash wasn't found. */
+		virtual RowHandle FindIndexedRow(IndexHash Index) const = 0;
+
 		/**
 		 * Runs a previously created query. This version takes an arbitrary query, but is limited to running queries that do not directly
 		 * access data from rows such as count queries.
