TEDS-Outliner: Add support for a separate widget purpose for the Item Label Column
StarshipCoreStyle: Add icon badges to signify data has overrides from the base class

diff --git a/Engine/Plugins/Experimental/TypedElementsDataStorage/Source/TypedElementsDataStorageUI/Private/TypedElementOutliner/Compatibility/TEDSCompatibilityUtils.cpp b/Engine/Plugins/Experimental/TypedElementsDataStorage/Source/TypedElementsDataStorageUI/Private/TypedElementOutliner/Compatibility/TEDSCompatibilityUtils.cpp
--- a/Engine/Plugins/Experimental/TypedElementsDataStorage/Source/TypedElementsDataStorageUI/Private/TypedElementOutliner/Compatibility/TEDSCompatibilityUtils.cpp
+++ b/Engine/Plugins/Experimental/TypedElementsDataStorage/Source/TypedElementsDataStorageUI/Private/TypedElementOutliner/Compatibility/TEDSCompatibilityUtils.cpp
@@ -13,6 +13,7 @@
 #include "Widgets/SWidget.h"
 
 FBaseTEDSOutlinerMode::FBaseTEDSOutlinerMode()
+: WidgetPurposes{TEXT("SceneOutliner.ItemLabel.Cell"), TEXT("SceneOutliner.Cell"), TEXT("General.Cell")}
 {
 	// Initialize the TEDS constructs
 	UTypedElementRegistry* Registry = UTypedElementRegistry::GetInstance();
@@ -36,20 +37,31 @@ FBaseTEDSOutlinerMode::FBaseTEDSOutlinerMode()
 
 		TSharedPtr<FTypedElementWidgetConstructor> OutWidgetConstructorPtr;
 
-		StorageUi->CreateWidgetConstructors(TEXT("General.Cell"), MatchApproach::ExactMatch, ColumnTypes, {},
-				[&OutWidgetConstructorPtr, ColumnTypes](
+		bool bFoundWidget = false;
+		
+		for(const FName& WidgetPurpose : WidgetPurposes)
+		{
+			StorageUi->CreateWidgetConstructors(WidgetPurpose, MatchApproach::ExactMatch, ColumnTypes, {},
+				[&OutWidgetConstructorPtr, ColumnTypes, &bFoundWidget](
 				TUniquePtr<FTypedElementWidgetConstructor> CreatedConstructor, 
 				TConstArrayView<TWeakObjectPtr<const UScriptStruct>> MatchedColumnTypes)
 				{
 					if (ColumnTypes.Num() == MatchedColumnTypes.Num())
 					{
 						OutWidgetConstructorPtr = TSharedPtr<FTypedElementWidgetConstructor>(CreatedConstructor.Release());
+						bFoundWidget = true;
 					}
 					// Either this was the exact match so no need to search further or the longest possible chain didn't match so the next ones will 
 					// always be shorter in both cases just return.
 					return false;
 				});
 
+			if(bFoundWidget)
+			{
+				break;
+			}
+		}
+
 		return OutWidgetConstructorPtr;
 	};
 
diff --git a/Engine/Plugins/Experimental/TypedElementsDataStorage/Source/TypedElementsDataStorageUI/Private/TypedElementOutliner/TypedElementOutlinerColumnIntegration.cpp b/Engine/Plugins/Experimental/TypedElementsDataStorage/Source/TypedElementsDataStorageUI/Private/TypedElementOutliner/TypedElementOutlinerColumnIntegration.cpp
--- a/Engine/Plugins/Experimental/TypedElementsDataStorage/Source/TypedElementsDataStorageUI/Private/TypedElementOutliner/TypedElementOutlinerColumnIntegration.cpp
+++ b/Engine/Plugins/Experimental/TypedElementsDataStorage/Source/TypedElementsDataStorageUI/Private/TypedElementOutliner/TypedElementOutlinerColumnIntegration.cpp
@@ -239,6 +239,16 @@ public:
 		return SNullWidget::NullWidget;
 	}
 
+	virtual void PopulateSearchStrings( const ISceneOutlinerTreeItem& Item, TArray< FString >& OutSearchStrings ) const override
+	{
+		// TODO: We don't currently have a way to convert TEDS widgets into searchable strings, but we can rely on the fallback column if it exists
+		if(FallbackColumn)
+		{
+			FallbackColumn->PopulateSearchStrings(Item, OutSearchStrings);
+		}
+	}
+	
+
 	TArray<TWeakObjectPtr<const UScriptStruct>> ColumnTypes;
 	TSharedPtr<FTypedElementWidgetConstructor> HeaderWidgetConstructor;
 	TSharedPtr<FTypedElementWidgetConstructor> CellWidgetConstructor;
@@ -284,6 +294,13 @@ void UTypedElementSceneOutlinerFactory::RegisterWidgetPurposes(ITypedElementData
 		LOCTEXT("CellWidgetPurpose", "Widgets for cells in any Scene Outliner for specific columns or column combinations."));
 	DataStorageUi.RegisterWidgetPurpose(FTypedElementSceneOutlinerQueryBinder::DefaultCellWidgetPurpose, PurposeType::UniqueByName,
 		LOCTEXT("DefaultCellWidgetPurpose", "The default widget to use in cells for the Scene Outliner."));
+
+	DataStorageUi.RegisterWidgetPurpose(FTypedElementSceneOutlinerQueryBinder::ItemLabelCellWidgetPurpose, PurposeType::UniqueByNameAndColumn,
+		LOCTEXT("ItemCellWidgetPurpose", "Widgets for cells in any Scene Outliner that are specific to the Item label column."));
+	DataStorageUi.RegisterWidgetPurpose(FTypedElementSceneOutlinerQueryBinder::DefaultItemLabelCellWidgetPurpose, PurposeType::UniqueByName,
+		LOCTEXT("DefaultItemCellWidgetPurpose", "The default widget to use in cells for the Scene Outliner specific to the Item label column."));
+
+
 }
 
 
@@ -297,6 +314,8 @@ const FName FTypedElementSceneOutlinerQueryBinder::HeaderWidgetPurpose(TEXT("Sce
 const FName FTypedElementSceneOutlinerQueryBinder::DefaultHeaderWidgetPurpose(TEXT("SceneOutliner.Header.Default"));
 const FName FTypedElementSceneOutlinerQueryBinder::CellWidgetPurpose(TEXT("SceneOutliner.Cell"));
 const FName FTypedElementSceneOutlinerQueryBinder::DefaultCellWidgetPurpose(TEXT("SceneOutliner.Cell.Default"));
+const FName FTypedElementSceneOutlinerQueryBinder::ItemLabelCellWidgetPurpose(TEXT("SceneOutliner.ItemLabel.Cell"));
+const FName FTypedElementSceneOutlinerQueryBinder::DefaultItemLabelCellWidgetPurpose(TEXT("SceneOutliner.ItemLabel.Cell.Default"));
 
 FTypedElementSceneOutlinerQueryBinder::FTypedElementSceneOutlinerQueryBinder()
 {
diff --git a/Engine/Plugins/Experimental/TypedElementsDataStorage/Source/TypedElementsDataStorageUI/Public/TypedElementOutliner/Compatibility/TEDSCompatibilityUtils.h b/Engine/Plugins/Experimental/TypedElementsDataStorage/Source/TypedElementsDataStorageUI/Public/TypedElementOutliner/Compatibility/TEDSCompatibilityUtils.h
--- a/Engine/Plugins/Experimental/TypedElementsDataStorage/Source/TypedElementsDataStorageUI/Public/TypedElementOutliner/Compatibility/TEDSCompatibilityUtils.h
+++ b/Engine/Plugins/Experimental/TypedElementsDataStorage/Source/TypedElementsDataStorageUI/Public/TypedElementOutliner/Compatibility/TEDSCompatibilityUtils.h
@@ -30,4 +30,5 @@ protected:
 	ITypedElementDataStorageCompatibilityInterface* StorageCompatibility{ nullptr };
 
 	TArray<TPair<TypedElementDataStorage::FQueryDescription, TSharedPtr<FTypedElementWidgetConstructor>>> QueryToWidgetConstructorMap;
+	TArray<FName, TFixedAllocator<3>> WidgetPurposes;
 };
diff --git a/Engine/Plugins/Experimental/TypedElementsDataStorage/Source/TypedElementsDataStorageUI/Public/TypedElementOutliner/TypedElementOutlinerColumnIntegration.h b/Engine/Plugins/Experimental/TypedElementsDataStorage/Source/TypedElementsDataStorageUI/Public/TypedElementOutliner/TypedElementOutlinerColumnIntegration.h
--- a/Engine/Plugins/Experimental/TypedElementsDataStorage/Source/TypedElementsDataStorageUI/Public/TypedElementOutliner/TypedElementOutlinerColumnIntegration.h
+++ b/Engine/Plugins/Experimental/TypedElementsDataStorage/Source/TypedElementsDataStorageUI/Public/TypedElementOutliner/TypedElementOutlinerColumnIntegration.h
@@ -56,6 +56,9 @@ public:
 	static const FName DefaultHeaderWidgetPurpose;
 	static const FName CellWidgetPurpose;
 	static const FName DefaultCellWidgetPurpose;
+	static const FName ItemLabelCellWidgetPurpose;
+	static const FName DefaultItemLabelCellWidgetPurpose;
+
 
 	static FTypedElementSceneOutlinerQueryBinder& GetInstance();
 
diff --git a/Engine/Source/Runtime/SlateCore/Private/Styling/StarshipCoreStyle.cpp b/Engine/Source/Runtime/SlateCore/Private/Styling/StarshipCoreStyle.cpp
--- a/Engine/Source/Runtime/SlateCore/Private/Styling/StarshipCoreStyle.cpp
+++ b/Engine/Source/Runtime/SlateCore/Private/Styling/StarshipCoreStyle.cpp
@@ -411,6 +411,19 @@ TSharedRef<ISlateStyle> FStarshipCoreStyle::Create()
 		Style->Set("Icons.Recent", new IMAGE_BRUSH_SVG("Starship/Common/Recent", Icon16x16));
 
 		Style->Set("Icons.BadgeModified", new IMAGE_BRUSH_SVG("Starship/Common/badge-modified", Icon16x16));
+
+		// Background images for all override badges (first layer)
+		Style->Set("Icons.BadgeOverrideBase", new IMAGE_BRUSH_SVG("Starship/Common/CircleBadgeDot", Icon16x16, FStyleColors::AccentBlue));
+		Style->Set("Icons.BadgeOverrideAddBase", new IMAGE_BRUSH_SVG("Starship/Common/CircleBadgeDot", Icon16x16, FStyleColors::AccentGreen));
+		Style->Set("Icons.BadgeOverrideRemoveBase", new IMAGE_BRUSH_SVG("Starship/Common/CircleBadgeDot", Icon16x16, FStyleColors::AccentRed));
+
+		// Foreground images for all override badges (second layer)
+		Style->Set("Icons.BadgeOverrideSelf", new IMAGE_BRUSH_SVG("Starship/Common/CircleBadgeBorder", Icon16x16, FStyleColors::Recessed));
+		Style->Set("Icons.BadgeOverrideChild", new IMAGE_BRUSH_SVG("Starship/Common/CircleBadgeOutline", Icon16x16, FStyleColors::Recessed));
+		Style->Set("Icons.BadgeOverrideSelfAndChild", new IMAGE_BRUSH_SVG("Starship/Common/CircleBadgeHalf", Icon16x16, FStyleColors::Recessed));
+		Style->Set("Icons.BadgeOverrideAdd", new IMAGE_BRUSH_SVG("Starship/Common/CircleBadgePlus", Icon16x16, FStyleColors::Recessed));
+		Style->Set("Icons.BadgeOverrideRemove", new IMAGE_BRUSH_SVG("Starship/Common/CircleBadgeMinus", Icon16x16, FStyleColors::Recessed));
+		
 		
 		// Toolbar Size Icons
 		Style->Set("Icons.Toolbar.Settings", new IMAGE_BRUSH_SVG("Starship/Common/settings", Icon20x20));
