TEDS: Fix a bug in the TypeInfo widget where the correct icon would not show up for actors

diff --git a/Engine/Plugins/Experimental/TypedElementsDataStorage/Source/TypedElementsDataStorageUI/Private/TypedElementOutliner/Widgets/TypedElementTypeInfoWidget.cpp b/Engine/Plugins/Experimental/TypedElementsDataStorage/Source/TypedElementsDataStorageUI/Private/TypedElementOutliner/Widgets/TypedElementTypeInfoWidget.cpp
--- a/Engine/Plugins/Experimental/TypedElementsDataStorage/Source/TypedElementsDataStorageUI/Private/TypedElementOutliner/Widgets/TypedElementTypeInfoWidget.cpp
+++ b/Engine/Plugins/Experimental/TypedElementsDataStorage/Source/TypedElementsDataStorageUI/Private/TypedElementOutliner/Widgets/TypedElementTypeInfoWidget.cpp
@@ -127,7 +127,7 @@ const FSlateBrush* FTypedElementTypeInfoWidgetConstructor::GetIconForRow(ITypedE
 	 */
 	
 	const UClass* Type = TypeInfoColumn->TypeInfo.Get();
-	FName IconName = Type->GetFName();
+	FName IconName;
 
 	// Allow the actor the first chance to provide an icon override
 	FMassActorFragment* ActorStore = DataStorage->GetColumn<FMassActorFragment>(Row);
@@ -139,6 +139,11 @@ const FSlateBrush* FTypedElementTypeInfoWidgetConstructor::GetIconForRow(ITypedE
 		}
 	}
 
+	if(IconName == NAME_None)
+	{
+		IconName = Type->GetFName();
+	}
+
 	// Check the cache if we already found an icon for this class
 	if(const FSlateBrush** CachedBrush = CachedIconMap.Find(IconName))
 	{
