TEDS Outliner - Add the ability for search text to be highlighted properly in the columns by adding a SetHighlightText method to TypedElementUITextCapability

diff --git a/Engine/Plugins/Experimental/TypedElementsDataStorage/Source/TypedElementsDataStorageUI/Private/TypedElementOutliner/TypedElementOutlinerColumnIntegration.cpp b/Engine/Plugins/Experimental/TypedElementsDataStorage/Source/TypedElementsDataStorageUI/Private/TypedElementOutliner/TypedElementOutlinerColumnIntegration.cpp
--- a/Engine/Plugins/Experimental/TypedElementsDataStorage/Source/TypedElementsDataStorageUI/Private/TypedElementOutliner/TypedElementOutlinerColumnIntegration.cpp
+++ b/Engine/Plugins/Experimental/TypedElementsDataStorage/Source/TypedElementsDataStorageUI/Private/TypedElementOutliner/TypedElementOutlinerColumnIntegration.cpp
@@ -21,6 +21,7 @@
 #include "SceneOutlinerFwd.h"
 #include "SceneOutlinerModule.h"
 #include "SceneOutlinerPublicTypes.h"
+#include "Elements/Interfaces/Capabilities/TypedElementUiTextCapability.h"
 #include "TypedElementOutliner/TypedElementOutlinerItem.h"
 
 #define LOCTEXT_NAMESPACE "TypedElementsUI_SceneOutliner"
@@ -120,7 +121,7 @@ public:
 		TSharedPtr<FTypedElementWidgetConstructor> InHeaderWidgetConstructor,
 		TSharedPtr<FTypedElementWidgetConstructor> InCellWidgetConstructor,
 		FName InFallbackColumnName,
-		ISceneOutliner& OwningOutliner)
+		ISceneOutliner& InOwningOutliner)
 		: ColumnTypes(MoveTemp(InColumnTypes))
 		, HeaderWidgetConstructor(MoveTemp(InHeaderWidgetConstructor))
 		, CellWidgetConstructor(MoveTemp(InCellWidgetConstructor))
@@ -129,6 +130,7 @@ public:
 		, StorageCompatibility(InStorageCompatibility)
 		, QueryDescription(Storage.GetQueryDescription(InQuery))
 		, NameId(InNameId)
+		, OwningOutliner(InOwningOutliner)
 	{
 		MetaData.AddOrSetMutableData(TEXT("Name"), NameId.ToString());
 		ColumnTypes.Shrink();
@@ -186,6 +188,22 @@ public:
 				]
 			];
 	}
+
+	void SetHighlightText(SWidget& Widget)
+	{
+		if (TSharedPtr<ITypedElementUiTextCapability> TextCapability = Widget.GetMetaData<ITypedElementUiTextCapability>())
+		{
+			TextCapability->SetHighlightText(OwningOutliner.GetFilterHighlightText());
+		}
+	
+		if (FChildren* ChildWidgets = Widget.GetChildren())
+		{
+			ChildWidgets->ForEachWidget([this](SWidget& ChildWidget)
+				{
+					SetHighlightText(ChildWidget);
+				});
+		}
+	}
 	
 	const TSharedRef<SWidget> ConstructRowWidget(FSceneOutlinerTreeItemRef TreeItem, const STableRow<FSceneOutlinerTreeItemPtr>& Row) override
 	{
@@ -193,6 +211,8 @@ public:
 		
 		TypedElementRowHandle RowHandle = TypedElementInvalidRowHandle;
 
+		TSharedPtr<SWidget> RowWidget;
+
 		if(const FTypedElementOutlinerTreeItem* TEDSItem = TreeItem->CastTo<FTypedElementOutlinerTreeItem>())
 		{
 			RowHandle = TEDSItem->GetRowHandle();
@@ -207,7 +227,7 @@ public:
 		}
 		else if(FallbackColumn)
 		{
-			return FallbackColumn->ConstructRowWidget(TreeItem, Row);
+			RowWidget = FallbackColumn->ConstructRowWidget(TreeItem, Row);
 		}
 
 		if (RowHandle != TypedElementInvalidRowHandle && Storage.HasColumns(RowHandle, ColumnTypes))
@@ -229,13 +249,20 @@ public:
 			if (TSharedPtr<SWidget> Widget = StorageUi.ConstructWidget(UiRowHandle, *CellWidgetConstructor, 
 				FComboMetaDataView(FGenericMetaDataView(MetaData)).Next(FQueryMetaDataView(QueryDescription))))
 			{
-				return Widget.ToSharedRef();
+				RowWidget =  Widget;
 			}
 			else
 			{
 				Storage.RemoveRow(RowHandle);
 			}
 		}
+
+		if(RowWidget)
+		{
+			SetHighlightText(*RowWidget);
+			return RowWidget.ToSharedRef();
+		}
+		
 		return SNullWidget::NullWidget;
 	}
 
@@ -259,6 +286,7 @@ public:
 	TypedElementDataStorage::FMetaData MetaData;
 	FName NameId;
 	TSharedPtr<ISceneOutlinerColumn> FallbackColumn;
+	ISceneOutliner& OwningOutliner;
 };
 
 
diff --git a/Engine/Plugins/Experimental/TypedElementsDataStorage/Source/TypedElementsDataStorageUI/Private/TypedElementOutliner/Widgets/TypedElementTypeInfoWidget.cpp b/Engine/Plugins/Experimental/TypedElementsDataStorage/Source/TypedElementsDataStorageUI/Private/TypedElementOutliner/Widgets/TypedElementTypeInfoWidget.cpp
--- a/Engine/Plugins/Experimental/TypedElementsDataStorage/Source/TypedElementsDataStorageUI/Private/TypedElementOutliner/Widgets/TypedElementTypeInfoWidget.cpp
+++ b/Engine/Plugins/Experimental/TypedElementsDataStorage/Source/TypedElementsDataStorageUI/Private/TypedElementOutliner/Widgets/TypedElementTypeInfoWidget.cpp
@@ -8,6 +8,7 @@
 #include "Elements/Columns/TypedElementMiscColumns.h"
 #include "Elements/Columns/TypedElementSlateWidgetColumns.h"
 #include "Elements/Columns/TypedElementTypeInfoColumns.h"
+#include "Elements/Interfaces/Capabilities/TypedElementUiTextCapability.h"
 #include "Styling/SlateIconFinder.h"
 #include "Widgets/Text/STextBlock.h"
 #include "Widgets/Images/SImage.h"
@@ -38,6 +39,8 @@ FTypedElementTypeInfoWidgetConstructor::FTypedElementTypeInfoWidgetConstructor(c
 TSharedPtr<SWidget> FTypedElementTypeInfoWidgetConstructor::CreateWidget(
 	const TypedElementDataStorage::FMetaDataView& Arguments	)
 {
+	bUseIcon = false;
+	
 	// Check if the caller provided metadata to use an icon widget
 	TypedElementDataStorage::FMetaDataEntryView MetaDataEntryView = Arguments.FindGeneric("TypedElementTypeInfoWidget_bUseIcon");
 	if(MetaDataEntryView.IsSet())
@@ -101,11 +104,13 @@ bool FTypedElementTypeInfoWidgetConstructor::FinalizeWidget(ITypedElementDataSto
 			// If not, we simply show a text block with the type
 			if(!ActualWidget)
 			{
-				ActualWidget = SNew(STextBlock)
+				TSharedPtr<STextBlock> TextBlock = SNew(STextBlock)
 								.ColorAndOpacity(FSlateColor::UseSubduedForeground())
 								.Text(FText::FromString(TypeInfoColumn->TypeInfo.Get()->GetName()));
-			}
 
+				TextBlock->AddMetadata(MakeShared<TTypedElementUiTextCapability<STextBlock>>(*TextBlock));
+				ActualWidget = TextBlock;
+			}
 
 			WidgetInstance->AddSlot()
 						.AutoWidth()
diff --git a/Engine/Source/Runtime/Slate/Private/Widgets/Text/SInlineEditableTextBlock.cpp b/Engine/Source/Runtime/Slate/Private/Widgets/Text/SInlineEditableTextBlock.cpp
--- a/Engine/Source/Runtime/Slate/Private/Widgets/Text/SInlineEditableTextBlock.cpp
+++ b/Engine/Source/Runtime/Slate/Private/Widgets/Text/SInlineEditableTextBlock.cpp
@@ -178,6 +178,12 @@ void SInlineEditableTextBlock::SetText( const FString& InText )
 	SetEditableText( Text );
 }
 
+
+void SInlineEditableTextBlock::SetHighlightText( const TAttribute< FText >& InText )
+{
+	TextBlock->SetHighlightText(InText);
+}
+
 void SInlineEditableTextBlock::SetWrapTextAt( const TAttribute<float>& InWrapTextAt )
 {
 	TextBlock->SetWrapTextAt( InWrapTextAt );
diff --git a/Engine/Source/Runtime/Slate/Public/Widgets/Text/SInlineEditableTextBlock.h b/Engine/Source/Runtime/Slate/Public/Widgets/Text/SInlineEditableTextBlock.h
--- a/Engine/Source/Runtime/Slate/Public/Widgets/Text/SInlineEditableTextBlock.h
+++ b/Engine/Source/Runtime/Slate/Public/Widgets/Text/SInlineEditableTextBlock.h
@@ -158,6 +158,8 @@ class SInlineEditableTextBlock: public SCompoundWidget
 	SLATE_API void SetText( const TAttribute< FText >& InText );
 	SLATE_API void SetText( const FString& InText );
 
+	SLATE_API void SetHighlightText(const TAttribute<FText>& InText);
+
 	/** Sets the wrap text at attribute.  See WrapTextAt attribute */
 	SLATE_API void SetWrapTextAt(const TAttribute<float>& InWrapTextAt);
 
diff --git a/Engine/Source/Runtime/TypedElementFramework/Public/Elements/Interfaces/Capabilities/TypedElementUiTextCapability.h b/Engine/Source/Runtime/TypedElementFramework/Public/Elements/Interfaces/Capabilities/TypedElementUiTextCapability.h
--- a/Engine/Source/Runtime/TypedElementFramework/Public/Elements/Interfaces/Capabilities/TypedElementUiTextCapability.h
+++ b/Engine/Source/Runtime/TypedElementFramework/Public/Elements/Interfaces/Capabilities/TypedElementUiTextCapability.h
@@ -17,6 +17,8 @@ public:
 	~ITypedElementUiTextCapability() override = default;
 
 	virtual void SetText(TAttribute<FText> Text) = 0;
+
+	virtual void SetHighlightText(TAttribute<FText> Text) = 0;
 };
 
 template<typename WidgetType>
@@ -30,6 +32,11 @@ public:
 		Widget.SetText(Text);
 	}
 
+	void SetHighlightText(TAttribute<FText> Text) override
+	{
+		Widget.SetHighlightText(Text);
+	}
+
 private:
 	WidgetType& Widget;
 };
\ No newline at end of file
