Adding initial support for hierarchies in TEDS.

This change includes adds two new columns, one to hold a parent relation to another row and one for the same in case the parent couldn't be found. A new query will check for the latter once per frame and replace it automatically with the former once the reference is found.

diff --git a/Engine/Plugins/Experimental/TypedElementsDataStorage/Source/TypedElementsDataStorage/Private/Queries/TypedElementHierarchyQueries.cpp b/Engine/Plugins/Experimental/TypedElementsDataStorage/Source/TypedElementsDataStorage/Private/Queries/TypedElementHierarchyQueries.cpp
new file mode 100644
--- /dev/null
+++ b/Engine/Plugins/Experimental/TypedElementsDataStorage/Source/TypedElementsDataStorage/Private/Queries/TypedElementHierarchyQueries.cpp
@@ -0,0 +1,30 @@
+// Copyright Epic Games, Inc. All Rights Reserved.
+
+#include "Queries/TypedElementHierarchyQueries.h"
+
+#include "Elements/Columns/TypedElementHiearchyColumns.h"
+#include "Elements/Framework/TypedElementQueryBuilder.h"
+
+void UTypedElementHiearchyQueriesFactory::RegisterQueries(ITypedElementDataStorageInterface& DataStorage) const
+{
+	using namespace TypedElementDataStorage;
+	using namespace TypedElementQueryBuilder;
+	using DSI = ITypedElementDataStorageInterface;
+
+	DataStorage.RegisterQuery(
+		Select(
+			TEXT("Resolve hierarchy rows"),
+			FProcessor(DSI::EQueryTickPhase::PrePhysics, DataStorage.GetQueryTickGroupName(DSI::EQueryTickGroups::Default)),
+			[](DSI::IQueryContext& Context, RowHandle Row, const FTypedElementUnresolvedParentColumn& UnresolvedParent)
+			{
+				RowHandle ParentRow = Context.FindIndexedRow(UnresolvedParent.ParentIdHash);
+				if (Context.IsRowAvailable(ParentRow))
+				{
+					Context.RemoveColumns<FTypedElementUnresolvedParentColumn>(Row);
+					Context.AddColumn(Row, FTypedElementParentColumn{ .Parent = ParentRow });
+				}
+			}
+		)
+		.Compile()
+	);
+}
diff --git a/Engine/Plugins/Experimental/TypedElementsDataStorage/Source/TypedElementsDataStorage/Private/Queries/TypedElementHierarchyQueries.h b/Engine/Plugins/Experimental/TypedElementsDataStorage/Source/TypedElementsDataStorage/Private/Queries/TypedElementHierarchyQueries.h
new file mode 100644
--- /dev/null
+++ b/Engine/Plugins/Experimental/TypedElementsDataStorage/Source/TypedElementsDataStorage/Private/Queries/TypedElementHierarchyQueries.h
@@ -0,0 +1,22 @@
+// Copyright Epic Games, Inc. All Rights Reserved.
+
+#pragma once
+
+#include "Elements/Interfaces/TypedElementDataStorageFactory.h"
+#include "UObject/ObjectMacros.h"
+
+#include "TypedElementHierarchyQueries.generated.h"
+
+/**
+ * Calls to queries for general hierarchy management.
+ */
+UCLASS()
+class UTypedElementHiearchyQueriesFactory : public UTypedElementDataStorageFactory
+{
+	GENERATED_BODY()
+
+public:
+	~UTypedElementHiearchyQueriesFactory() override = default;
+
+	void RegisterQueries(ITypedElementDataStorageInterface& DataStorage) const override;
+};
diff --git a/Engine/Source/Runtime/TypedElementFramework/Public/Elements/Columns/TypedElementHiearchyColumns.h b/Engine/Source/Runtime/TypedElementFramework/Public/Elements/Columns/TypedElementHiearchyColumns.h
new file mode 100644
--- /dev/null
+++ b/Engine/Source/Runtime/TypedElementFramework/Public/Elements/Columns/TypedElementHiearchyColumns.h
@@ -0,0 +1,33 @@
+// Copyright Epic Games, Inc. All Rights Reserved.
+
+#pragma once
+
+#include "Elements/Common/TypedElementHandles.h"
+#include "Elements/Interfaces/TypedElementDataStorageInterface.h"
+#include "UObject/ObjectMacros.h"
+
+#include "TypedElementHiearchyColumns.generated.h"
+
+/**
+ * A reference to the direct hierarchical parent of this row.
+ */
+USTRUCT(meta = (DisplayName = "Parent"))
+struct FTypedElementParentColumn final : public FTypedElementDataStorageColumn
+{
+	GENERATED_BODY()
+
+	TypedElementDataStorage::RowHandle Parent;
+};
+
+/**
+ * A reference to the direct hierarchical parent of this row which has not been resolved yet. The stored value will
+ * be used to attempt to find the indexed row. This column can not be used to find rows that are not indexed.
+ */
+USTRUCT(meta = (DisplayName = "Parent (Unresolved)"))
+struct FTypedElementUnresolvedParentColumn final : public FTypedElementDataStorageColumn
+{
+	GENERATED_BODY()
+
+	UPROPERTY(meta = (IgnoreForMemberInitializationTest))
+	uint64 ParentIdHash;
+};
