TEDS CIS fix: disable MT access detector on shipping targets

diff --git a/Engine/Plugins/Experimental/TypedElementsDataStorage/Source/TypedElementsDataStorage/Private/TypedElementDatabaseScratchBuffer.cpp b/Engine/Plugins/Experimental/TypedElementsDataStorage/Source/TypedElementsDataStorage/Private/TypedElementDatabaseScratchBuffer.cpp
--- a/Engine/Plugins/Experimental/TypedElementsDataStorage/Source/TypedElementsDataStorage/Private/TypedElementDatabaseScratchBuffer.cpp
+++ b/Engine/Plugins/Experimental/TypedElementsDataStorage/Source/TypedElementsDataStorage/Private/TypedElementDatabaseScratchBuffer.cpp
@@ -30,7 +30,7 @@ FTypedElementDatabaseScratchBuffer::~FTypedElementDatabaseScratchBuffer()
 void* FTypedElementDatabaseScratchBuffer::Allocate(size_t Size, size_t Alignment)
 {
 	// ReaderAccessDetector to assert that this scope has shared access to any thread also making an Allocate call
-	TScopedReaderAccessDetector Detector(AccessDetector);
+	UE_MT_SCOPED_READ_ACCESS(AccessDetector);
 	return GetThreadLocalBlockController().Allocate(Size, Alignment);
 }
 
@@ -39,7 +39,7 @@ void FTypedElementDatabaseScratchBuffer::RecycleBlocks()
 	if (FullBlocks)
 	{
 		// WriterAccessDetector to assert that this scope has exclusive access - should not overlap Allocate()
-		TScopedWriterDetector Detector(AccessDetector);
+		UE_MT_SCOPED_WRITE_ACCESS(AccessDetector);
 		
 		// Destroy any objects in blocks. It walks the allocations in a block backwards and invokes the destructors on the allocated 
 		// objects. It moves to the next block once it reaches the beginning of the block. At the end the cleared blocks are reinserted 
diff --git a/Engine/Plugins/Experimental/TypedElementsDataStorage/Source/TypedElementsDataStorage/Private/TypedElementDatabaseScratchBuffer.h b/Engine/Plugins/Experimental/TypedElementsDataStorage/Source/TypedElementsDataStorage/Private/TypedElementDatabaseScratchBuffer.h
--- a/Engine/Plugins/Experimental/TypedElementsDataStorage/Source/TypedElementsDataStorage/Private/TypedElementDatabaseScratchBuffer.h
+++ b/Engine/Plugins/Experimental/TypedElementsDataStorage/Source/TypedElementsDataStorage/Private/TypedElementDatabaseScratchBuffer.h
@@ -5,6 +5,7 @@
 #include <atomic>
 #include <type_traits>
 #include "GenericPlatform/GenericPlatformMath.h"
+#include "Math/UnrealMath.h"
 #include "Misc/MTAccessDetector.h"
 #include "Templates/AlignmentTemplates.h"
 #include "Templates/UnrealTemplate.h"
@@ -81,7 +82,7 @@ private:
 	void ConfigureDestructorTail(FDestructorTail& Destructor, DestructorFunction Callback, void* Object, int32 Count = 1);
 
 	// Detect any invalid calls to recycle blocks while allocating memory.
-	FRWRecursiveAccessDetector AccessDetector;
+	UE_MT_DECLARE_RW_ACCESS_DETECTOR(AccessDetector);
 
 	std::atomic<FBlock*> AvailableBlocks = nullptr;
 	std::atomic<FBlock*> FullBlocks = nullptr;
