Adding an object to TEDS now first checks if the object has already been added.

#nojira

diff --git a/Engine/Plugins/Experimental/TypedElementsDataStorage/Source/TypedElementsDataStorage/Private/TypedElementDatabaseCompatibility.cpp b/Engine/Plugins/Experimental/TypedElementsDataStorage/Source/TypedElementsDataStorage/Private/TypedElementDatabaseCompatibility.cpp
--- a/Engine/Plugins/Experimental/TypedElementsDataStorage/Source/TypedElementsDataStorage/Private/TypedElementDatabaseCompatibility.cpp
+++ b/Engine/Plugins/Experimental/TypedElementsDataStorage/Source/TypedElementsDataStorage/Private/TypedElementDatabaseCompatibility.cpp
@@ -341,12 +341,17 @@ void UTypedElementDatabaseCompatibility::CreateStandardArchetypes()
 
 bool UTypedElementDatabaseCompatibility::ShouldAddObject(const UObject* Object) const
 {
+	using namespace TypedElementDataStorage;
+
 	bool Include = true;
-	const ObjectRegistrationFilter* Filter = ObjectRegistrationFilters.GetData();
-	const ObjectRegistrationFilter* FilterEnd = Filter + ObjectRegistrationFilters.Num();
-	for (; Include && Filter != FilterEnd; ++Filter)
+	if (!Storage->IsRowAvailable(Storage->FindIndexedRow(GenerateIndexHash(Object))))
 	{
-		Include = (*Filter)(*this, Object);
+		const ObjectRegistrationFilter* Filter = ObjectRegistrationFilters.GetData();
+		const ObjectRegistrationFilter* FilterEnd = Filter + ObjectRegistrationFilters.Num();
+		for (; Include && Filter != FilterEnd; ++Filter)
+		{
+			Include = (*Filter)(*this, Object);
+		}
 	}
 	return Include;
 }
@@ -355,7 +360,7 @@ TypedElementRowHandle UTypedElementDatabaseCompatibility::DealiasObject(const UO
 {
 	for (const ObjectToRowDealiaser& Dealiaser : ObjectToRowDialiasers)
 	{
-		if (TypedElementRowHandle Row = Dealiaser(*this, Object); Row != TypedElementDataStorage::InvalidRowHandle)
+		if (TypedElementRowHandle Row = Dealiaser(*this, Object); Storage->IsRowAvailable(Row))
 		{
 			return Row;
 		}
