TEDS UI Widget Constructor pass.

The following was changed:
- CanBeReused() was removed. This ended up with no practical use and the current usage pattern keeps widget constructors around.
- Widget constructors are now property initialized. Before this change the arguments meant for the construction of a widget constructor didn't make it to the widget constructor.
- The columns that were matched to the widget constructor, if any, are now stored on the constructor for later reference.
- The label widget in TEDS UI was updated to use the newly available matched columns.
- Reduce the function export in the label widget to just functions instead of the entire class.

diff --git a/Engine/Plugins/Experimental/TypedElementsDataStorage/Source/TypedElementsDataStorage/Private/TypedElementDatabaseUI.cpp b/Engine/Plugins/Experimental/TypedElementsDataStorage/Source/TypedElementsDataStorage/Private/TypedElementDatabaseUI.cpp
--- a/Engine/Plugins/Experimental/TypedElementsDataStorage/Source/TypedElementsDataStorage/Private/TypedElementDatabaseUI.cpp
+++ b/Engine/Plugins/Experimental/TypedElementsDataStorage/Source/TypedElementsDataStorage/Private/TypedElementDatabaseUI.cpp
@@ -332,13 +332,24 @@ void UTypedElementDatabaseUi::ConstructWidgets(FName Purpose, const TypedElement
 		{
 			std::visit(Internal::TOverloaded
 				{
-					[this, &Arguments, &ConstructionCallback](const UScriptStruct* Constructor)
+					[this, &Arguments, &ConstructionCallback](const UScriptStruct* ConstructorType)
 					{ 
-						CreateWidgetInstanceFromDescription(Constructor, Arguments, ConstructionCallback); 
+						FTypedElementWidgetConstructor* Constructor = reinterpret_cast<FTypedElementWidgetConstructor*>(
+							FMemory_Alloca_Aligned(ConstructorType->GetStructureSize(), ConstructorType->GetMinAlignment()));
+						if (Constructor)
+						{
+							ConstructorType->InitializeStruct(Constructor);
+							CreateWidgetInstance(*Constructor, Arguments, ConstructionCallback);
+							ConstructorType->DestroyStruct(&Constructor);
+						}
+						else
+						{
+							checkf(false, TEXT("Remaining stack space is too small to create a Typed Elements widget constructor from a description."));
+						}
 					},
 					[this, &Arguments, &ConstructionCallback](const TUniquePtr<FTypedElementWidgetConstructor>& Constructor)
 					{
-						CreateWidgetInstanceFromInstance(Constructor.Get(), Arguments, ConstructionCallback);
+						CreateWidgetInstance(*Constructor, Arguments, ConstructionCallback);
 					}
 				}, Factory.Constructor);
 		}
@@ -348,7 +359,7 @@ void UTypedElementDatabaseUi::ConstructWidgets(FName Purpose, const TypedElement
 bool UTypedElementDatabaseUi::CreateSingleWidgetConstructor(
 	const FWidgetFactory::ConstructorType& Constructor,
 	const TypedElementDataStorage::FMetaDataView& Arguments,
-	TConstArrayView<TWeakObjectPtr<const UScriptStruct>> MatchedColumnTypes, 
+	TArray<TWeakObjectPtr<const UScriptStruct>> MatchedColumnTypes, 
 	const WidgetConstructorCallback& Callback)
 {
 	return std::visit(Internal::TOverloaded
@@ -360,7 +371,8 @@ bool UTypedElementDatabaseUi::CreateSingleWidgetConstructor(
 				if (Result)
 				{
 					Target->InitializeStruct(Result.Get());
-					return Callback(MoveTemp(Result), MatchedColumnTypes);
+					Result->Initialize(Arguments, MoveTemp(MatchedColumnTypes));
+					return Callback(MoveTemp(Result), Result->GetMatchedColumns());
 				}
 				return true;
 			},
@@ -374,61 +386,14 @@ bool UTypedElementDatabaseUi::CreateSingleWidgetConstructor(
 				{
 					TargetType->InitializeStruct(Result.Get());
 					TargetType->CopyScriptStruct(Result.Get(), Target.Get());
-					return Callback(MoveTemp(Result), MatchedColumnTypes);
+					Result->Initialize(Arguments, MoveTemp(MatchedColumnTypes));
+					return Callback(MoveTemp(Result), Result->GetMatchedColumns());
 				}
 				return true;
 			}
 		}, Constructor);
 }
 
-void UTypedElementDatabaseUi::CreateWidgetInstanceFromDescription(
-	const UScriptStruct* Target,
-	const TypedElementDataStorage::FMetaDataView& Arguments,
-	const WidgetCreatedCallback& ConstructionCallback)
-{
-	FTypedElementWidgetConstructor* Constructor = reinterpret_cast<FTypedElementWidgetConstructor*>(
-		FMemory_Alloca_Aligned(Target->GetStructureSize(), Target->GetMinAlignment()));
-	if (Constructor)
-	{
-		Target->InitializeStruct(Constructor);
-		CreateWidgetInstance(*Constructor, Arguments, ConstructionCallback);
-		Target->DestroyStruct(&Constructor);
-	}
-	else
-	{
-		checkf(false, TEXT("Remaining stack space is too small to create a Typed Elements widget constructor from a description."));
-	}
-}
-
-void UTypedElementDatabaseUi::CreateWidgetInstanceFromInstance(
-	FTypedElementWidgetConstructor* SourceConstructor,
-	const TypedElementDataStorage::FMetaDataView& Arguments,
-	const WidgetCreatedCallback& ConstructionCallback)
-{
-	if (SourceConstructor->CanBeReused())
-	{
-		CreateWidgetInstance(*SourceConstructor, Arguments, ConstructionCallback);
-	}
-	else
-	{
-		const UScriptStruct* Target = SourceConstructor->GetTypeInfo();
-		checkf(Target, TEXT("Expected valid type information from a widget constructor."));
-		FTypedElementWidgetConstructor* Constructor = reinterpret_cast<FTypedElementWidgetConstructor*>(
-			FMemory_Alloca_Aligned(Target->GetStructureSize(), Target->GetMinAlignment()));
-		if (Constructor)
-		{
-			Target->InitializeStruct(Constructor);
-			Target->CopyScriptStruct(Constructor, SourceConstructor);
-			CreateWidgetInstance(*Constructor, Arguments, ConstructionCallback);
-			Target->DestroyStruct(&Constructor);
-		}
-		else
-		{
-			checkf(false, TEXT("Remaining stack space is too small to create a Typed Elements widget constructor from an instance."));
-		}
-	}
-}
-
 void UTypedElementDatabaseUi::CreateWidgetInstance(
 	FTypedElementWidgetConstructor& Constructor, 
 	const TypedElementDataStorage::FMetaDataView& Arguments,
@@ -450,30 +415,7 @@ void UTypedElementDatabaseUi::CreateWidgetInstance(
 TSharedPtr<SWidget> UTypedElementDatabaseUi::ConstructWidget(TypedElementRowHandle Row, FTypedElementWidgetConstructor& Constructor,
 	const TypedElementDataStorage::FMetaDataView& Arguments)
 {
-	if (Constructor.CanBeReused())
-	{
-		return Constructor.Construct(Row, Storage, this, Arguments);
-	}
-	else
-	{
-		const UScriptStruct* Target = Constructor.GetTypeInfo();
-		FTypedElementWidgetConstructor* ConstructorCopy = reinterpret_cast<FTypedElementWidgetConstructor*>(
-			FMemory_Alloca_Aligned(Target->GetStructureSize(), Target->GetMinAlignment()));
-		if (ConstructorCopy)
-		{
-			Target->InitializeStruct(ConstructorCopy);
-			Target->CopyScriptStruct(ConstructorCopy, &Constructor);
-			TSharedPtr<SWidget> Widget = ConstructorCopy->Construct(Row, Storage, this, Arguments);
-			Target->DestroyStruct(&ConstructorCopy);
-			return Widget;
-		}
-		else
-		{
-			checkf(false, TEXT("Remaining stack space is too small to create a Typed Elements widget constructor from an instance."));
-			return nullptr;
-		}
-		return nullptr;
-	}
+	return Constructor.Construct(Row, Storage, this, Arguments);
 }
 
 void UTypedElementDatabaseUi::ListWidgetPurposes(const WidgetPurposeCallback& Callback) const
@@ -500,7 +442,6 @@ void UTypedElementDatabaseUi::CreateWidgetConstructors_LongestMatch(const TArray
 	TArray<TWeakObjectPtr<const UScriptStruct>> MatchedColumns;
 	for (auto FactoryIt = WidgetFactories.CreateConstIterator(); FactoryIt && !Columns.IsEmpty(); ++FactoryIt)
 	{
-		int32 MatchingIndex = INDEX_NONE;
 		if (FactoryIt->Columns.MinimumColumnMatchRequired() > Columns.Num())
 		{
 			// There are more columns required for this factory than there are in the requested columns list so skip this
@@ -514,27 +455,29 @@ void UTypedElementDatabaseUi::CreateWidgetConstructors_LongestMatch(const TArray
 		{
 			// Remove the found columns from the requested list.
 			Algo::SortBy(MatchedColumns, [](const TWeakObjectPtr<const UScriptStruct>& Column) { return Column.Get(); });
-			TConstArrayView<TWeakObjectPtr<const UScriptStruct>> CleanMatchedColumns(MatchedColumns.GetData(), Algo::Unique(MatchedColumns));
-
+			MatchedColumns.SetNum(Algo::Unique(MatchedColumns), false);
+			
 			TWeakObjectPtr<const UScriptStruct>* ColumnsIt = Columns.GetData();
 			TWeakObjectPtr<const UScriptStruct>* ColumnsEnd = ColumnsIt + Columns.Num();
 			int32 ColumnIndex = 0;
-			for (const TWeakObjectPtr<const UScriptStruct>& MatchedColumn : CleanMatchedColumns)
+			for (const TWeakObjectPtr<const UScriptStruct>& MatchedColumn : MatchedColumns)
 			{
+				// Remove all the columns that were matched from the provided column list.
 				while (*ColumnsIt != MatchedColumn)
 				{
 					++ColumnIndex;
 					++ColumnsIt;
 					if (ColumnsIt == ColumnsEnd)
 					{
-						ensureMsgf(false, TEXT("A previous found matching column can't be found in the original array."));
+						ensureMsgf(false, TEXT("A previously found matching column can't be found in the original array."));
 						return;
 					}
 				}
-				Columns.RemoveAt(ColumnIndex);
+				Columns.RemoveAt(ColumnIndex, 1, false /* Allow shrinking */);
+				--ColumnsEnd;
 			}
 			
-			if (!CreateSingleWidgetConstructor(FactoryIt->Constructor, Arguments, CleanMatchedColumns, Callback))
+			if (!CreateSingleWidgetConstructor(FactoryIt->Constructor, Arguments, MoveTemp(MatchedColumns), Callback))
 			{
 				return;
 			}
@@ -562,29 +505,11 @@ void UTypedElementDatabaseUi::CreateWidgetConstructors_ExactMatch(const TArray<F
 		if (Factory.Columns.Verify(MatchedColumns, Columns, true))
 		{
 			Algo::SortBy(MatchedColumns, [](const TWeakObjectPtr<const UScriptStruct>& Column) { return Column.Get(); });
-			TWeakObjectPtr<const UScriptStruct>* MatchedColumnsIt = MatchedColumns.GetData();
-			TWeakObjectPtr<const UScriptStruct>* MatchedColumnsEnd = MatchedColumnsIt + MatchedColumns.Num();
-			bool bFullyMatched = true;
-			for (TWeakObjectPtr<const UScriptStruct> Column : Columns)
-			{
-				while (*MatchedColumnsIt != Column)
-				{
-					if (++MatchedColumnsIt == MatchedColumnsEnd)
-					{
-						bFullyMatched = false;
-						break;
-					}
-				}
-				if (!bFullyMatched)
-				{
-					break;
-				}
-			}
-
-			if (bFullyMatched)
+			MatchedColumns.SetNum(Algo::Unique(MatchedColumns), false);
+			if (MatchedColumns.Num() == Columns.Num())
 			{
 				Columns.Reset();
-				CreateSingleWidgetConstructor(Factory.Constructor, Arguments, MatchedColumns, Callback);
+				CreateSingleWidgetConstructor(Factory.Constructor, Arguments, MoveTemp(MatchedColumns), Callback);
 				return;
 			}
 		}
@@ -618,7 +543,7 @@ void UTypedElementDatabaseUi::CreateWidgetConstructors_SingleMatch(const TArray<
 			if (ColumnData[0] == Columns[ColumnIndex])
 			{
 				Columns.RemoveAt(ColumnIndex);
-				CreateSingleWidgetConstructor((*FactoryIt).Constructor, Arguments, ColumnData, Callback);
+				CreateSingleWidgetConstructor((*FactoryIt).Constructor, Arguments, TArray<TWeakObjectPtr<const UScriptStruct>>(ColumnData), Callback);
 				// Match was found so move on to the next column in the column.
 				break;
 			}
diff --git a/Engine/Plugins/Experimental/TypedElementsDataStorage/Source/TypedElementsDataStorage/Private/TypedElementDatabaseUI.h b/Engine/Plugins/Experimental/TypedElementsDataStorage/Source/TypedElementsDataStorage/Private/TypedElementDatabaseUI.h
--- a/Engine/Plugins/Experimental/TypedElementsDataStorage/Source/TypedElementsDataStorage/Private/TypedElementDatabaseUI.h
+++ b/Engine/Plugins/Experimental/TypedElementsDataStorage/Source/TypedElementsDataStorage/Private/TypedElementDatabaseUI.h
@@ -79,19 +79,9 @@ private:
 	bool CreateSingleWidgetConstructor(
 		const FWidgetFactory::ConstructorType& Constructor,
 		const TypedElementDataStorage::FMetaDataView& Arguments,
-		TConstArrayView<TWeakObjectPtr<const UScriptStruct>> MatchedColumnTypes,
+		TArray<TWeakObjectPtr<const UScriptStruct>> MatchedColumnTypes,
 		const WidgetConstructorCallback& Callback);
 
-	void CreateWidgetInstanceFromDescription(
-		const UScriptStruct* Target,
-		const TypedElementDataStorage::FMetaDataView& Arguments,
-		const WidgetCreatedCallback& ConstructionCallback);
-
-	void CreateWidgetInstanceFromInstance(
-		FTypedElementWidgetConstructor* SourceConstructor,
-		const TypedElementDataStorage::FMetaDataView& Arguments,
-		const WidgetCreatedCallback& ConstructionCallback);
-
 	void CreateWidgetInstance(
 		FTypedElementWidgetConstructor& Constructor,
 		const TypedElementDataStorage::FMetaDataView& Arguments,
diff --git a/Engine/Plugins/Experimental/TypedElementsDataStorage/Source/TypedElementsDataStorageUI/Private/TypedElementLabelWidget.cpp b/Engine/Plugins/Experimental/TypedElementsDataStorage/Source/TypedElementsDataStorageUI/Private/TypedElementLabelWidget.cpp
--- a/Engine/Plugins/Experimental/TypedElementsDataStorage/Source/TypedElementsDataStorageUI/Private/TypedElementLabelWidget.cpp
+++ b/Engine/Plugins/Experimental/TypedElementsDataStorage/Source/TypedElementsDataStorageUI/Private/TypedElementLabelWidget.cpp
@@ -75,7 +75,6 @@ void UTypedElementLabelWidgetFactory::RegisterQueries(ITypedElementDataStorageIn
 {
 	using namespace TypedElementQueryBuilder;
 	using DSI = ITypedElementDataStorageInterface;
-	
 
 	TypedElementQueryHandle UpdateLabelWidget = DataStorage.RegisterQuery(
 		Select()
@@ -131,9 +130,7 @@ void UTypedElementLabelWidgetFactory::RegisterWidgetConstructors(ITypedElementDa
 	using namespace TypedElementQueryBuilder;
 
 	DataStorageUi.RegisterWidgetFactory<FTypedElementLabelWidgetConstructor>(FName(TEXT("General.Cell")), 
-		FColumn<FTypedElementLabelColumn>());
-	DataStorageUi.RegisterWidgetFactory<FTypedElementLabelWithHashTooltipWidgetConstructor>(FName(TEXT("General.Cell")),
-		FColumn<FTypedElementLabelColumn>() && FColumn<FTypedElementLabelHashColumn>());
+		FColumn<FTypedElementLabelColumn>() || (FColumn<FTypedElementLabelColumn>() && FColumn<FTypedElementLabelHashColumn>()));
 }
 
 
@@ -228,7 +225,7 @@ TSharedPtr<SWidget> FTypedElementLabelWidgetConstructor::Construct(
 bool FTypedElementLabelWidgetConstructor::SetColumns(ITypedElementDataStorageInterface* DataStorage, TypedElementRowHandle Row)
 {
 	DataStorage->GetColumn<FTypedElementU64IntValueCacheColumn>(Row)->Value = 0;
-	DataStorage->AddOrGetColumn<FTypedElementLabelWidgetColumn>(Row)->bShowHashInTooltip = false;
+	DataStorage->AddOrGetColumn<FTypedElementLabelWidgetColumn>(Row)->bShowHashInTooltip = (MatchedColumnTypes.Num() == 2);
 	return true;
 }
 
@@ -242,39 +239,7 @@ bool FTypedElementLabelWidgetConstructor::FinalizeWidget(
 		DataStorage, 
 		DataStorage->GetColumn<FTypedElementRowReferenceColumn>(Row)->Row,
 		*DataStorage->GetColumn<FTypedElementU64IntValueCacheColumn>(Row),
-		Widget, false);
-	return true;
-}
-
-
-
-//
-// FTypedElementLabelWithHashTooltipWidgetConstructor
-//
-
-FTypedElementLabelWithHashTooltipWidgetConstructor::FTypedElementLabelWithHashTooltipWidgetConstructor()
-	: Super(FTypedElementLabelWithHashTooltipWidgetConstructor::StaticStruct())
-{
-}
-
-bool FTypedElementLabelWithHashTooltipWidgetConstructor::SetColumns(ITypedElementDataStorageInterface* DataStorage, TypedElementRowHandle Row)
-{
-	DataStorage->GetColumn<FTypedElementU64IntValueCacheColumn>(Row)->Value = 0;
-	DataStorage->AddOrGetColumn<FTypedElementLabelWidgetColumn>(Row)->bShowHashInTooltip = true;
-	return true;
-}
-
-bool FTypedElementLabelWithHashTooltipWidgetConstructor::FinalizeWidget(
-	ITypedElementDataStorageInterface* DataStorage,
-	ITypedElementDataStorageUiInterface* DataStorageUi,
-	TypedElementRowHandle Row,
-	const TSharedPtr<SWidget>& Widget)
-{
-	SyncColumnsToWidget(
-		DataStorage,
-		DataStorage->GetColumn<FTypedElementRowReferenceColumn>(Row)->Row,
-		*DataStorage->GetColumn<FTypedElementU64IntValueCacheColumn>(Row),
-		Widget, true);
+		Widget, MatchedColumnTypes.Num() == 2);
 	return true;
 }
 
diff --git a/Engine/Plugins/Experimental/TypedElementsDataStorage/Source/TypedElementsDataStorageUI/Public/TypedElementLabelWidget.h b/Engine/Plugins/Experimental/TypedElementsDataStorage/Source/TypedElementsDataStorageUI/Public/TypedElementLabelWidget.h
--- a/Engine/Plugins/Experimental/TypedElementsDataStorage/Source/TypedElementsDataStorageUI/Public/TypedElementLabelWidget.h
+++ b/Engine/Plugins/Experimental/TypedElementsDataStorage/Source/TypedElementsDataStorageUI/Public/TypedElementLabelWidget.h
@@ -11,30 +11,30 @@
 #include "TypedElementLabelWidget.generated.h"
 
 UCLASS()
-class TYPEDELEMENTSDATASTORAGEUI_API UTypedElementLabelWidgetFactory : public UTypedElementDataStorageFactory
+class UTypedElementLabelWidgetFactory : public UTypedElementDataStorageFactory
 {
 	GENERATED_BODY()
 
 public:
 	~UTypedElementLabelWidgetFactory() override = default;
 
-	void RegisterQueries(ITypedElementDataStorageInterface& DataStorage) const override;
-	void RegisterWidgetConstructors(ITypedElementDataStorageInterface& DataStorage,
+	TYPEDELEMENTSDATASTORAGEUI_API void RegisterQueries(ITypedElementDataStorageInterface& DataStorage) const override;
+	TYPEDELEMENTSDATASTORAGEUI_API void RegisterWidgetConstructors(ITypedElementDataStorageInterface& DataStorage,
 		ITypedElementDataStorageUiInterface& DataStorageUi) const override;
 };
 
 USTRUCT()
-struct TYPEDELEMENTSDATASTORAGEUI_API FTypedElementLabelWidgetConstructor : public FTypedElementWidgetConstructor
+struct FTypedElementLabelWidgetConstructor : public FTypedElementWidgetConstructor
 {
 	GENERATED_BODY()
 
 public:
-	FTypedElementLabelWidgetConstructor();
+	TYPEDELEMENTSDATASTORAGEUI_API FTypedElementLabelWidgetConstructor();
 	~FTypedElementLabelWidgetConstructor() override = default;
 
-	TConstArrayView<const UScriptStruct*> GetAdditionalColumnsList() const override;
+	TYPEDELEMENTSDATASTORAGEUI_API TConstArrayView<const UScriptStruct*> GetAdditionalColumnsList() const override;
 
-	TSharedPtr<SWidget> Construct(
+	TYPEDELEMENTSDATASTORAGEUI_API TSharedPtr<SWidget> Construct(
 		TypedElementRowHandle Row, 
 		ITypedElementDataStorageInterface* DataStorage,
 		ITypedElementDataStorageUiInterface* DataStorageUi,
@@ -47,25 +47,10 @@ protected:
 		TypedElementRowHandle Row, const TSharedPtr<SWidget>& Widget) override;
 };
 
-USTRUCT()
-struct TYPEDELEMENTSDATASTORAGEUI_API FTypedElementLabelWithHashTooltipWidgetConstructor : public FTypedElementLabelWidgetConstructor
-{
-	GENERATED_BODY()
-
-public:
-	FTypedElementLabelWithHashTooltipWidgetConstructor();
-	~FTypedElementLabelWithHashTooltipWidgetConstructor() = default;
-
-protected:
-	bool SetColumns(ITypedElementDataStorageInterface* DataStorage, TypedElementRowHandle Row) override;
-	bool FinalizeWidget(ITypedElementDataStorageInterface* DataStorage, ITypedElementDataStorageUiInterface* DataStorageUi,
-		TypedElementRowHandle Row, const TSharedPtr<SWidget>& Widget) override;
-};
-
 USTRUCT(meta = (DisplayName = "Label widget"))
-struct TYPEDELEMENTSDATASTORAGEUI_API FTypedElementLabelWidgetColumn : public FTypedElementDataStorageColumn
+struct FTypedElementLabelWidgetColumn : public FTypedElementDataStorageColumn
 {
 	GENERATED_BODY()
 
-	bool bShowHashInTooltip : 1{ false };
+	bool bShowHashInTooltip { false };
 };
\ No newline at end of file
diff --git a/Engine/Source/Runtime/TypedElementFramework/Private/Elements/Interfaces/TypedElementDataStorageUiInterface.cpp b/Engine/Source/Runtime/TypedElementFramework/Private/Elements/Interfaces/TypedElementDataStorageUiInterface.cpp
--- a/Engine/Source/Runtime/TypedElementFramework/Private/Elements/Interfaces/TypedElementDataStorageUiInterface.cpp
+++ b/Engine/Source/Runtime/TypedElementFramework/Private/Elements/Interfaces/TypedElementDataStorageUiInterface.cpp
@@ -9,19 +9,26 @@ FTypedElementWidgetConstructor::FTypedElementWidgetConstructor(const UScriptStru
 {
 }
 
+bool FTypedElementWidgetConstructor::Initialize(const TypedElementDataStorage::FMetaDataView& InArguments,
+	TArray<TWeakObjectPtr<const UScriptStruct>> InMatchedColumnTypes)
+{
+	MatchedColumnTypes = MoveTemp(InMatchedColumnTypes);
+	return true;
+}
+
 const UScriptStruct* FTypedElementWidgetConstructor::GetTypeInfo() const
 {
 	return TypeInfo;
 }
 
-TConstArrayView<const UScriptStruct*> FTypedElementWidgetConstructor::GetAdditionalColumnsList() const
+const TArray<TWeakObjectPtr<const UScriptStruct>>& FTypedElementWidgetConstructor::GetMatchedColumns() const
 {
-	return {};
+	return MatchedColumnTypes;
 }
 
-bool FTypedElementWidgetConstructor::CanBeReused() const
+TConstArrayView<const UScriptStruct*> FTypedElementWidgetConstructor::GetAdditionalColumnsList() const
 {
-	return true;
+	return {};
 }
 
 TSharedPtr<SWidget> FTypedElementWidgetConstructor::Construct(
diff --git a/Engine/Source/Runtime/TypedElementFramework/Public/Elements/Interfaces/TypedElementDataStorageUiInterface.h b/Engine/Source/Runtime/TypedElementFramework/Public/Elements/Interfaces/TypedElementDataStorageUiInterface.h
--- a/Engine/Source/Runtime/TypedElementFramework/Public/Elements/Interfaces/TypedElementDataStorageUiInterface.h
+++ b/Engine/Source/Runtime/TypedElementFramework/Public/Elements/Interfaces/TypedElementDataStorageUiInterface.h
@@ -29,20 +29,18 @@ public:
 
 	virtual ~FTypedElementWidgetConstructor() = default;
 
+	/** Initializes a new constructor based on the provided arguments.. */
+	TYPEDELEMENTFRAMEWORK_API virtual bool Initialize(const TypedElementDataStorage::FMetaDataView& InArguments,
+		TArray<TWeakObjectPtr<const UScriptStruct>> InMatchedColumnTypes);
+
 	/** Retrieves the type information for the constructor type. */
 	TYPEDELEMENTFRAMEWORK_API virtual const UScriptStruct* GetTypeInfo() const;
+	/** Retrieves the columns, if any, that were matched to this constructor when it was created. */
+	TYPEDELEMENTFRAMEWORK_API virtual const TArray<TWeakObjectPtr<const UScriptStruct>>& GetMatchedColumns() const;
 
-	/** Returns a list of additional columns the widget requires. */
+	/** Returns a list of additional columns the widget requires to be added to its rows. */
 	TYPEDELEMENTFRAMEWORK_API virtual TConstArrayView<const UScriptStruct*> GetAdditionalColumnsList() const;
 
-	/** 
-	 * Whether or not an instance of this constructor can be reused. Setting this to true means that the Data Storage will
-	 * avoid recreating new instances of the constructor and instead keeps reusing the created version. The default is
-	 * 'false'. Override this function and set it to true if there's no internal state that could interfere with creating
-	 * multiple widgets from the same constructor, e.g. data stored from passed in arguments.
-	 */
-	TYPEDELEMENTFRAMEWORK_API virtual bool CanBeReused() const;
-
 	/**
 	 * Constructs the widget according to the provided information. Information is collected by calling
 	 * the below functions CreateWidget and AddColumns. It's recommended to overload those
@@ -70,6 +68,7 @@ protected:
 		TypedElementRowHandle Row,
 		const TSharedPtr<SWidget>& Widget);
 
+	TArray<TWeakObjectPtr<const UScriptStruct>> MatchedColumnTypes;
 	const UScriptStruct* TypeInfo{ nullptr };
 };
 
@@ -189,12 +188,18 @@ public:
 	virtual bool RegisterWidgetFactory(FName Purpose, TUniquePtr<FTypedElementWidgetConstructor>&& Constructor, 
 		TypedElementQueryBuilder::FQueryConditions Columns) = 0;
 	
-	/** Creates widget constructors for the requested purpose. */
+	/** 
+	 * Creates widget constructors for the requested purpose.
+	 * The provided arguments will be used to configure the constructor. Settings made this way will be applied to all
+	 * widgets created from the constructor, if applicable.
+	 */
 	virtual void CreateWidgetConstructors(FName Purpose, 
 		const TypedElementDataStorage::FMetaDataView& Arguments, const WidgetConstructorCallback& Callback) = 0;
 	/** 
 	 * Finds matching widget constructors for provided columns, preferring longer matches over shorter matches.
 	 * The provided list of columns will be updated to contain all columns that couldn't be matched.
+	 * The provided arguments will be used to configure the constructor. Settings made this way will be applied to all
+	 * widgets created from the constructor, if applicable.
 	 */
 	virtual void CreateWidgetConstructors(FName Purpose, EMatchApproach MatchApproach, TArray<TWeakObjectPtr<const UScriptStruct>>& Columns,
 		const TypedElementDataStorage::FMetaDataView& Arguments, const WidgetConstructorCallback& Callback) = 0;
@@ -207,7 +212,11 @@ public:
 	virtual void ConstructWidgets(FName Purpose, const TypedElementDataStorage::FMetaDataView& Arguments,
 		const WidgetCreatedCallback& ConstructionCallback) = 0;
 
-	/** Creates a single widget using the provided constructor. Arguments can optionally be used to intialize the constructor. */
+	/** 
+	 * Creates a single widget using the provided constructor. 
+	 * The provided row will be used to store the widget information on. If columns have already been added to the row, the 
+	 * constructor is free to use that to configure the widget. Arguments are used by the constructor to configure the widget.
+	 */
 	virtual TSharedPtr<SWidget> ConstructWidget(TypedElementRowHandle Row, FTypedElementWidgetConstructor& Constructor,
 		const TypedElementDataStorage::FMetaDataView& Arguments) = 0;
 
