Fix for a rare crash when a large number of objects are deleted in two consecutive frames.

At the root of the problem of the failure is a memory stomp in TEDS' scratch buffer. The scratch buffer assumes that all work has completed meaning no more of its memory is in use. Mass however queues commands during processing up for next frame so the scratch buffer releases memory too early in some case. The following criteria have to be met for the crash to trigger:
- Before the TEDS tick a large enough number of commands with extra data need to be created, at least enough to fill a scratch buffer chunk.
- At the start of the TEDS tick the filled memory chunks are recycled. This is the bug as the chunks may still be in use. This is by itself not enough to trigger the crash as the chunks are not cleared so the bug doesn't happen until the blocks are reused.
- During the TEDS tick new commands need to be issued before the TEDS tick reached a point where the chunk data is processed. This corrupts the original data to a degree that will trigger the crash.

The solution to the problem is to keep a frame counter in the scratch buffer and a last touched frame number in each chunk. During recycling only the chunks that were touched in the last frame or earlier can be recycled, while the rest will have to wait for another frame.

diff --git a/Engine/Plugins/Experimental/TypedElementsDataStorage/Source/TypedElementsDataStorage/Private/TypedElementDatabase.cpp b/Engine/Plugins/Experimental/TypedElementsDataStorage/Source/TypedElementsDataStorage/Private/TypedElementDatabase.cpp
--- a/Engine/Plugins/Experimental/TypedElementsDataStorage/Source/TypedElementsDataStorage/Private/TypedElementDatabase.cpp
+++ b/Engine/Plugins/Experimental/TypedElementsDataStorage/Source/TypedElementsDataStorage/Private/TypedElementDatabase.cpp
@@ -69,7 +69,8 @@ void UTypedElementDatabase::Initialize()
 	check(GEditor);
 	UMassEntityEditorSubsystem* Mass = GEditor->GetEditorSubsystem<UMassEntityEditorSubsystem>();
 	check(Mass);
-	Mass->GetOnPreTickDelegate().AddUObject(this, &UTypedElementDatabase::OnPreMassTick);
+	OnPreMassTickHandle = Mass->GetOnPreTickDelegate().AddUObject(this, &UTypedElementDatabase::OnPreMassTick);
+	OnPostMassTickHandle = Mass->GetOnPostTickDelegate().AddUObject(this, &UTypedElementDatabase::OnPostMassTick);
 
 	ActiveEditorEntityManager = Mass->GetMutableEntityManager();
 	ActiveEditorPhaseManager = Mass->GetMutablePhaseManager();
@@ -113,10 +114,19 @@ void UTypedElementDatabase::Deinitialize()
 void UTypedElementDatabase::OnPreMassTick(float DeltaTime)
 {
 	checkf(IsAvailable(), TEXT("Typed Element Database was ticked while it's not ready."));
+	
 	OnUpdateDelegate.Broadcast();
 	// Process pending commands after other systems have had a chance to update. Other systems may have executed work needed
 	// to complete pending work.
 	ProcessPendingCommands();
+
+	Environment->GetScratchBuffer().NextFrame();
+}
+
+void UTypedElementDatabase::OnPostMassTick(float DeltaTime)
+{
+	checkf(IsAvailable(), TEXT("Typed Element Database was ticked while it's not ready."));
+	
 	// Recycle any full scratch memory blocks from the previous frame.
 	Environment->GetScratchBuffer().RecycleBlocks();
 }
@@ -867,6 +877,14 @@ void UTypedElementDatabase::FinalizePhase(EQueryTickPhase Phase, float DeltaTime
 
 void UTypedElementDatabase::Reset()
 {
+	if (UMassEntityEditorSubsystem* Mass = GEditor->GetEditorSubsystem<UMassEntityEditorSubsystem>())
+	{
+		Mass->GetOnPostTickDelegate().Remove(OnPostMassTickHandle);
+		Mass->GetOnPreTickDelegate().Remove(OnPreMassTickHandle);
+	}
+	OnPostMassTickHandle.Reset();
+	OnPreMassTickHandle.Reset();
+
 	if (ActiveEditorPhaseManager)
 	{
 		Queries.Clear(*ActiveEditorPhaseManager.Get());
diff --git a/Engine/Plugins/Experimental/TypedElementsDataStorage/Source/TypedElementsDataStorage/Private/TypedElementDatabase.h b/Engine/Plugins/Experimental/TypedElementsDataStorage/Source/TypedElementsDataStorage/Private/TypedElementDatabase.h
--- a/Engine/Plugins/Experimental/TypedElementsDataStorage/Source/TypedElementsDataStorage/Private/TypedElementDatabase.h
+++ b/Engine/Plugins/Experimental/TypedElementsDataStorage/Source/TypedElementsDataStorage/Private/TypedElementDatabase.h
@@ -32,8 +32,10 @@ public:
 	void Initialize();
 	void Deinitialize();
 
-	/** Triggered just before underlying Mass processing gets ticked */
+	/** Triggered at the start of the underlying Mass' tick cycle. */
 	void OnPreMassTick(float DeltaTime);
+	/** Triggered just before underlying Mass processing completes it's tick cycle. */
+	void OnPostMassTick(float DeltaTime);
 
 	TSharedPtr<FMassEntityManager> GetActiveMutableEditorEntityManager();
 	TSharedPtr<const FMassEntityManager> GetActiveEditorEntityManager() const;
@@ -147,6 +149,8 @@ private:
 	FTypedElementExtendedQueryStore Queries;
 	
 	FTypedElementOnDataStorageUpdate OnUpdateDelegate;
+	FDelegateHandle OnPreMassTickHandle;
+	FDelegateHandle OnPostMassTickHandle;
 
 	TSharedPtr<FMassEntityManager> ActiveEditorEntityManager;
 	TSharedPtr<FMassProcessingPhaseManager> ActiveEditorPhaseManager;
diff --git a/Engine/Plugins/Experimental/TypedElementsDataStorage/Source/TypedElementsDataStorage/Private/TypedElementDatabaseScratchBuffer.cpp b/Engine/Plugins/Experimental/TypedElementsDataStorage/Source/TypedElementsDataStorage/Private/TypedElementDatabaseScratchBuffer.cpp
--- a/Engine/Plugins/Experimental/TypedElementsDataStorage/Source/TypedElementsDataStorage/Private/TypedElementDatabaseScratchBuffer.cpp
+++ b/Engine/Plugins/Experimental/TypedElementsDataStorage/Source/TypedElementsDataStorage/Private/TypedElementDatabaseScratchBuffer.cpp
@@ -31,7 +31,12 @@ void* FTypedElementDatabaseScratchBuffer::Allocate(size_t Size, size_t Alignment
 {
 	// ReaderAccessDetector to assert that this scope has shared access to any thread also making an Allocate call
 	UE_MT_SCOPED_READ_ACCESS(AccessDetector);
-	return GetThreadLocalBlockController().Allocate(Size, Alignment);
+	return GetThreadLocalBlockController().Allocate(Size, Alignment, FrameId);
+}
+
+void FTypedElementDatabaseScratchBuffer::NextFrame()
+{
+	FrameId++;
 }
 
 void FTypedElementDatabaseScratchBuffer::RecycleBlocks()
@@ -47,26 +52,44 @@ void FTypedElementDatabaseScratchBuffer::RecycleBlocks()
 		// relatively simple to destroy. The dirty blocks are added in front of the available blocks instead of the end to avoid having
 		// to iterate over the entire list of available blocks.
 
+		FBlock* RemainingBlocks = nullptr;
 		FBlock* DirtyBlock = FullBlocks;
-		FBlock* LastDirtyBlock = FullBlocks;
 		do
 		{
-			FDestructorTail* Destructor = DirtyBlock->DestructionTail;
-			while (Destructor)
+			// If the block hasn't been touched for at least one frame it can be deleted.
+			if (DirtyBlock->LastTouchedByFrame < FrameId)
 			{
-				Destructor->Destructor(reinterpret_cast<char*>(Destructor) - Destructor->StructOffset, Destructor->InstanceCount);
-				Destructor = Destructor->PreviousTail;
+				// Call destructors on any memory that needs it.
+				FDestructorTail* Destructor = DirtyBlock->DestructionTail;
+				while (Destructor)
+				{
+					Destructor->Destructor(reinterpret_cast<char*>(Destructor) - Destructor->StructOffset, Destructor->InstanceCount);
+					Destructor = Destructor->PreviousTail;
+				}
+				DirtyBlock->DestructionTail = nullptr;
+
+				// The block is now clean so set the next block to process.
+				FBlock* CleanBlock = DirtyBlock;
+				DirtyBlock = DirtyBlock->NextBlock;
+
+				// Reinsert the clean block into the chain of available blocks.
+				CleanBlock->NextBlock.store(AvailableBlocks);
+				AvailableBlocks = CleanBlock;
 			}
-			DirtyBlock->DestructionTail = nullptr;
+			else
+			{
+				// The block can't be processed yet so move onto the next block.
+				FBlock* DelayedBlock = DirtyBlock;
+				DirtyBlock = DirtyBlock->NextBlock;
 
-			LastDirtyBlock = DirtyBlock;
-			DirtyBlock = DirtyBlock->NextBlock;
+				// Store the delayed blocks for future processing.
+				DelayedBlock->NextBlock = RemainingBlocks;
+				RemainingBlocks = DelayedBlock;
+			}
 		} while (DirtyBlock);
 		
-		// Attach the full blocks in front of the available blocks.
-		LastDirtyBlock->NextBlock.store(AvailableBlocks);
-		AvailableBlocks.store(FullBlocks);
-		FullBlocks = nullptr;
+		// Set the full blocks to the first block that couldn't be removed or null if all filled up blocks have been recycled.
+		FullBlocks = RemainingBlocks;
 	}
 }
 
@@ -99,7 +122,7 @@ FTypedElementDatabaseScratchBuffer::FBlockController::~FBlockController()
 	RecycleBlock();
 }
 
-void* FTypedElementDatabaseScratchBuffer::FBlockController::Allocate(size_t Size, size_t Alignment)
+void* FTypedElementDatabaseScratchBuffer::FBlockController::Allocate(size_t Size, size_t Alignment, uint64 FrameId)
 {
 	checkf(Alignment <= alignof(FBlock), TEXT("Alignment of %i for allocation in database scratch buffer exceeds maximal alignment of %i."),
 		static_cast<int>(Alignment), static_cast<int>(alignof(FBlock)));
@@ -118,6 +141,7 @@ void* FTypedElementDatabaseScratchBuffer::FBlockController::Allocate(size_t Size
 			NewFront = Size;
 		}
 
+		Block->LastTouchedByFrame = FrameId;
 		Block->Front = NewFront;
 		return Block->Buffer + Index;
 	}
@@ -133,7 +157,7 @@ void* FTypedElementDatabaseScratchBuffer::FBlockController::Allocate(size_t Size
 		void* ExtendedBuffer = FMemory::Malloc(Size, Alignment);
 		
 		// Add an entry to the block with the sole purpose of deleting the extended buffer.
-		void* BufferStoreAddress = Allocate(sizeof(FExtendedBufferStore), alignof(FExtendedBufferStore));
+		void* BufferStoreAddress = Allocate(sizeof(FExtendedBufferStore), alignof(FExtendedBufferStore), FrameId);
 		FExtendedBufferStore* BufferStore = reinterpret_cast<FExtendedBufferStore*>(BufferStoreAddress);
 		BufferStore->ExtendedBuffer = ExtendedBuffer;
 
diff --git a/Engine/Plugins/Experimental/TypedElementsDataStorage/Source/TypedElementsDataStorage/Private/TypedElementDatabaseScratchBuffer.h b/Engine/Plugins/Experimental/TypedElementsDataStorage/Source/TypedElementsDataStorage/Private/TypedElementDatabaseScratchBuffer.h
--- a/Engine/Plugins/Experimental/TypedElementsDataStorage/Source/TypedElementsDataStorage/Private/TypedElementDatabaseScratchBuffer.h
+++ b/Engine/Plugins/Experimental/TypedElementsDataStorage/Source/TypedElementsDataStorage/Private/TypedElementDatabaseScratchBuffer.h
@@ -27,9 +27,12 @@ public:
 	template<typename T, typename... ArgTypes>
 	T* EmplaceArray(int32 Count, const ArgTypes&... Args);
 
+	// Increments the frame id by one. The frame id is used to track which blocks were used in the current frame and
+	// can't therefore be safely removed. Call this at the start of a frame.
+	void NextFrame();
 	// Makes previously filled up blocks available again. This can only safely be called when Allocate(...) can not
 	// be called. As allocations should only be called when Mass is running processors, it means this function
-	// can be safely called when the processors are not running, i.e. at the end of a tick.
+	// can be safely called when the processors are not running.
 	void RecycleBlocks();
 
 	constexpr static int32 MaxAllocationSize();
@@ -54,6 +57,8 @@ private:
 		constexpr static size_t BlockSize = 64 * 1024; // 64kb blocks
 		
 		char Buffer[BlockSize]; // Keep at top for alignment.
+		// Indicates the last frame this block was used in
+		uint64 LastTouchedByFrame = 0;
 		FDestructorTail* DestructionTail = nullptr;
 		std::atomic<FBlock*> NextBlock = nullptr;
 		std::atomic<uint32> Owner = 0;
@@ -67,7 +72,7 @@ private:
 		explicit FBlockController(FTypedElementDatabaseScratchBuffer& InOwner);
 		~FBlockController();
 
-		void* Allocate(size_t Size, size_t Alignment);
+		void* Allocate(size_t Size, size_t Alignment, uint64 FrameId);
 		FBlock* GetEmptyBlock();
 		void RecycleBlock();
 		void ConfigureDestructorTail(FDestructorTail& Destructor, DestructorFunction Callback, void* Object, int32 Count = 1);
@@ -89,6 +94,9 @@ private:
 	// Running counter so each FBlockController gets a unique id to identify the blocks assigned to them. Zero is reserved 
 	// to indicate it's not in use by a block controller.
 	std::atomic<uint32> BlockControllerId = 1;
+
+	// The number that identifies the current frame. Used to keep track of the blocks that were used in this frame or earlier.
+	uint64 FrameId = 0;
 };
 
 
diff --git a/Engine/Plugins/Runtime/MassEntity/Source/MassEntityEditor/Private/MassEntityEditorSubsystem.cpp b/Engine/Plugins/Runtime/MassEntity/Source/MassEntityEditor/Private/MassEntityEditorSubsystem.cpp
--- a/Engine/Plugins/Runtime/MassEntity/Source/MassEntityEditor/Private/MassEntityEditorSubsystem.cpp
+++ b/Engine/Plugins/Runtime/MassEntity/Source/MassEntityEditor/Private/MassEntityEditorSubsystem.cpp
@@ -141,6 +141,8 @@ void UMassEntityEditorSubsystem::Tick(float DeltaTime)
 		CompletionEvent->Wait();
 	}
 
+	OnPostTickDelegate.Broadcast(DeltaTime);
+
 	IsProcessing = false;
 }
 
diff --git a/Engine/Plugins/Runtime/MassEntity/Source/MassEntityEditor/Public/MassEntityEditorSubsystem.h b/Engine/Plugins/Runtime/MassEntity/Source/MassEntityEditor/Public/MassEntityEditorSubsystem.h
--- a/Engine/Plugins/Runtime/MassEntity/Source/MassEntityEditor/Public/MassEntityEditorSubsystem.h
+++ b/Engine/Plugins/Runtime/MassEntity/Source/MassEntityEditor/Public/MassEntityEditorSubsystem.h
@@ -20,6 +20,7 @@ class MASSENTITYEDITOR_API UMassEntityEditorSubsystem : public UEditorSubsystem
 
 public:
 	DECLARE_MULTICAST_DELEGATE_OneParam(FOnPreTick, float /*DeltaTime*/);
+	DECLARE_MULTICAST_DELEGATE_OneParam(FOnPostTick, float /*DeltaTime*/);
 
 	UMassEntityEditorSubsystem();
 	~UMassEntityEditorSubsystem();
@@ -28,6 +29,7 @@ public:
 	TSharedRef<FMassProcessingPhaseManager> GetMutablePhaseManager() { return PhaseManager; }
 
 	FOnPreTick& GetOnPreTickDelegate() { return OnPreTickDelegate; }
+	FOnPostTick& GetOnPostTickDelegate() { return OnPostTickDelegate; }
 
 	/** Registers a dynamic processor. This needs to be a fully formed processor and will be slotted in during the next tick. */
 	void RegisterDynamicProcessor(UMassProcessor& Processor);
@@ -54,6 +56,8 @@ protected:
 
 	/** Gets broadcast right before issuing FMassProcessingPhaseManager's phases execution */
 	FOnPreTick OnPreTickDelegate;
+	/** Gets broadcast right after completing FMassProcessingPhaseManager's phases execution */
+	FOnPostTick OnPostTickDelegate;
 	
 	bool IsProcessing{ false };
 };
