Fixed Mac compile error.

#nojira

diff --git a/Engine/Plugins/Experimental/TypedElementsDataStorage/Source/TypedElementsDataStorage/Private/TypedElementDatabaseScratchBuffer.cpp b/Engine/Plugins/Experimental/TypedElementsDataStorage/Source/TypedElementsDataStorage/Private/TypedElementDatabaseScratchBuffer.cpp
--- a/Engine/Plugins/Experimental/TypedElementsDataStorage/Source/TypedElementsDataStorage/Private/TypedElementDatabaseScratchBuffer.cpp
+++ b/Engine/Plugins/Experimental/TypedElementsDataStorage/Source/TypedElementsDataStorage/Private/TypedElementDatabaseScratchBuffer.cpp
@@ -122,7 +122,7 @@ FTypedElementDatabaseScratchBuffer::FBlockController::~FBlockController()
 	RecycleBlock();
 }
 
-void* FTypedElementDatabaseScratchBuffer::FBlockController::Allocate(size_t Size, size_t Alignment, uint64 FrameId)
+void* FTypedElementDatabaseScratchBuffer::FBlockController::Allocate(size_t Size, size_t Alignment, uint64 LocalFrameId)
 {
 	checkf(Alignment <= alignof(FBlock), TEXT("Alignment of %i for allocation in database scratch buffer exceeds maximal alignment of %i."),
 		static_cast<int>(Alignment), static_cast<int>(alignof(FBlock)));
@@ -141,7 +141,7 @@ void* FTypedElementDatabaseScratchBuffer::FBlockController::Allocate(size_t Size
 			NewFront = Size;
 		}
 
-		Block->LastTouchedByFrame = FrameId;
+		Block->LastTouchedByFrame = LocalFrameId;
 		Block->Front = NewFront;
 		return Block->Buffer + Index;
 	}
@@ -157,7 +157,7 @@ void* FTypedElementDatabaseScratchBuffer::FBlockController::Allocate(size_t Size
 		void* ExtendedBuffer = FMemory::Malloc(Size, Alignment);
 		
 		// Add an entry to the block with the sole purpose of deleting the extended buffer.
-		void* BufferStoreAddress = Allocate(sizeof(FExtendedBufferStore), alignof(FExtendedBufferStore), FrameId);
+		void* BufferStoreAddress = Allocate(sizeof(FExtendedBufferStore), alignof(FExtendedBufferStore), LocalFrameId);
 		FExtendedBufferStore* BufferStore = reinterpret_cast<FExtendedBufferStore*>(BufferStoreAddress);
 		BufferStore->ExtendedBuffer = ExtendedBuffer;
 
diff --git a/Engine/Plugins/Experimental/TypedElementsDataStorage/Source/TypedElementsDataStorage/Private/TypedElementDatabaseScratchBuffer.h b/Engine/Plugins/Experimental/TypedElementsDataStorage/Source/TypedElementsDataStorage/Private/TypedElementDatabaseScratchBuffer.h
--- a/Engine/Plugins/Experimental/TypedElementsDataStorage/Source/TypedElementsDataStorage/Private/TypedElementDatabaseScratchBuffer.h
+++ b/Engine/Plugins/Experimental/TypedElementsDataStorage/Source/TypedElementsDataStorage/Private/TypedElementDatabaseScratchBuffer.h
@@ -72,7 +72,7 @@ private:
 		explicit FBlockController(FTypedElementDatabaseScratchBuffer& InOwner);
 		~FBlockController();
 
-		void* Allocate(size_t Size, size_t Alignment, uint64 FrameId);
+		void* Allocate(size_t Size, size_t Alignment, uint64 LocalFrameId);
 		FBlock* GetEmptyBlock();
 		void RecycleBlock();
 		void ConfigureDestructorTail(FDestructorTail& Destructor, DestructorFunction Callback, void* Object, int32 Count = 1);
