Fix crash when >64 TypedElementDatastore queries are registered
* Container holding FTypedElementExtendedQuery is the TTypedElementHandleStore - a generational handle sparse array
* When FTypedElementExtendedQueryStore::RegisterQuery registered, query is passed down to FTypedElementQueryProcessorData::CommonQueryConfiguration which caches the query by pointer
* The fix in this CL removes the caching of queries by address, instead it holds onto them by handle similar to how subqueries are held in FQueryDescription
* Removed UB writing FTypedElementExtendedQueryStore::Handle::Handle and subsequently reading FTypedElementExtendedQueryStore::Handle::Data.
* Improved some const-correctness

diff --git a/Engine/Plugins/Experimental/TypedElementsDataStorage/Source/TypedElementsDataStorage/Private/Processors/TypedElementProcessorAdaptors.cpp b/Engine/Plugins/Experimental/TypedElementsDataStorage/Source/TypedElementsDataStorage/Private/Processors/TypedElementProcessorAdaptors.cpp
--- a/Engine/Plugins/Experimental/TypedElementsDataStorage/Source/TypedElementsDataStorage/Private/Processors/TypedElementProcessorAdaptors.cpp
+++ b/Engine/Plugins/Experimental/TypedElementsDataStorage/Source/TypedElementsDataStorage/Private/Processors/TypedElementProcessorAdaptors.cpp
@@ -521,8 +521,7 @@ public:
 
 	TypedElementDataStorage::FQueryResult RunQuery(TypedElementQueryHandle Query) override
 	{
-		FTypedElementExtendedQueryStore::Handle Handle;
-		Handle.Handle = Query;
+		const FTypedElementExtendedQueryStore::Handle Handle(Query);
 		// This can be safely called because there's not callback, which means no columns are accessed, even for select queries.
 		return QueryStore.RunQuery(Context.GetEntityManagerChecked(), Handle);
 	}
@@ -538,9 +537,9 @@ public:
 	{
 		if (SubqueryIndex < QueryDescription.Subqueries.Num())
 		{
-			FTypedElementExtendedQueryStore::Handle Handle;
-			Handle.Handle = QueryDescription.Subqueries[SubqueryIndex];
-			return QueryStore.RunQuery(Context.GetEntityManagerChecked(), Handle, Callback);
+			const TypedElementQueryHandle SubqueryHandle = QueryDescription.Subqueries[SubqueryIndex];
+			const FTypedElementExtendedQueryStore::Handle StorageHandle(SubqueryHandle);
+			return QueryStore.RunQuery(Context.GetEntityManagerChecked(), StorageHandle, Callback);
 		}
 		else
 		{
@@ -553,9 +552,9 @@ public:
 	{
 		if (SubqueryIndex < QueryDescription.Subqueries.Num())
 		{
-			FTypedElementExtendedQueryStore::Handle Handle;
-			Handle.Handle = QueryDescription.Subqueries[SubqueryIndex];
-			return QueryStore.RunQuery(Context.GetEntityManagerChecked(), Handle, Row, Callback);
+			const TypedElementQueryHandle SubqueryHandle = QueryDescription.Subqueries[SubqueryIndex];
+			const FTypedElementExtendedQueryStore::Handle StorageHandle(SubqueryHandle);
+			return QueryStore.RunQuery(Context.GetEntityManagerChecked(), StorageHandle, Row, Callback);
 		}
 		else
 		{
@@ -616,13 +615,14 @@ FTypedElementQueryProcessorData::FTypedElementQueryProcessorData(UMassProcessor&
 }
 
 bool FTypedElementQueryProcessorData::CommonQueryConfiguration(
-	UMassProcessor& InOwner, 
-	FTypedElementExtendedQuery& InQuery, 
-	FTypedElementExtendedQueryStore& InQueryStore, 
+	UMassProcessor& InOwner,
+	FTypedElementExtendedQuery& InQuery,
+	FTypedElementExtendedQueryStore::Handle InQueryHandle,
+	FTypedElementExtendedQueryStore& InQueryStore,
 	FTypedElementDatabaseEnvironment& InEnvironment,
 	TArrayView<FMassEntityQuery> Subqueries)
 {
-	ParentQuery = &InQuery;
+	ParentQuery = InQueryHandle;
 	QueryStore = &InQueryStore;
 	Environment = &InEnvironment;
 
@@ -633,8 +633,7 @@ bool FTypedElementQueryProcessorData::CommonQueryConfiguration(
 		int32 CurrentSubqueryIndex = 0;
 		for (TypedElementQueryHandle SubqueryHandle : InQuery.Description.Subqueries)
 		{
-			FTypedElementExtendedQueryStore::Handle SubqueryStoreHandle;
-			SubqueryStoreHandle.Handle = SubqueryHandle;
+			const FTypedElementExtendedQueryStore::Handle SubqueryStoreHandle(SubqueryHandle);
 			if (const FTypedElementExtendedQuery* Subquery = InQueryStore.Get(SubqueryStoreHandle))
 			{
 				if (ensureMsgf(Subquery->NativeQuery.CheckValidity(), TEXT("Provided subquery isn't valid. This can be because it couldn't be "
@@ -679,15 +678,16 @@ EMassProcessingPhase FTypedElementQueryProcessorData::MapToMassProcessingPhase(I
 
 FString FTypedElementQueryProcessorData::GetProcessorName() const
 {
-	return ParentQuery ? ParentQuery->Description.Callback.Name.ToString() : FString(TEXT("<unnamed>"));
+	const FTypedElementExtendedQuery* StoredQuery = QueryStore ? QueryStore->Get(ParentQuery) : nullptr;
+	return StoredQuery ? StoredQuery->Description.Callback.Name.ToString() : FString(TEXT("<unnamed>"));
 }
 
 bool FTypedElementQueryProcessorData::PrepareCachedDependenciesOnQuery(
 	ITypedElementDataStorageInterface::FQueryDescription& Description, FMassExecutionContext& Context)
 {
 	const int32 DependencyCount = Description.DependencyTypes.Num();
-	TWeakObjectPtr<const UClass>* Types = Description.DependencyTypes.GetData();
-	ITypedElementDataStorageInterface::EQueryDependencyFlags* Flags = Description.DependencyFlags.GetData();
+	const TWeakObjectPtr<const UClass>* Types = Description.DependencyTypes.GetData();
+	const ITypedElementDataStorageInterface::EQueryDependencyFlags* Flags = Description.DependencyFlags.GetData();
 	TWeakObjectPtr<UObject>* Caches = Description.CachedDependencies.GetData();
 
 	for (int32 Index = 0; Index < DependencyCount; ++Index)
@@ -785,9 +785,10 @@ TypedElementDataStorage::FQueryResult FTypedElementQueryProcessorData::Execute(
 
 void FTypedElementQueryProcessorData::Execute(FMassEntityManager& EntityManager, FMassExecutionContext& Context)
 {
-	checkf(ParentQuery, TEXT("A query callback was registered for execution without an associated query."));
+	FTypedElementExtendedQuery* StoredQuery = QueryStore->GetMutable(ParentQuery);
+	checkf(StoredQuery, TEXT("A query callback was registered for execution without an associated query."));
 	
-	ITypedElementDataStorageInterface::FQueryDescription& Description = ParentQuery->Description;
+	ITypedElementDataStorageInterface::FQueryDescription& Description = StoredQuery->Description;
 	NativeQuery.ForEachEntityChunk(EntityManager, Context,
 		[this, &Description](FMassExecutionContext& Context)
 		{
@@ -819,11 +820,12 @@ FMassEntityQuery& UTypedElementQueryProcessorCallbackAdapterProcessorBase::GetQu
 }
 
 bool UTypedElementQueryProcessorCallbackAdapterProcessorBase::ConfigureQueryCallback(
-	FTypedElementExtendedQuery& Query, 
-	FTypedElementExtendedQueryStore& QueryStore, 
+	FTypedElementExtendedQuery& Query,
+	FTypedElementExtendedQueryStore::Handle QueryHandle,
+	FTypedElementExtendedQueryStore& QueryStore,
 	FTypedElementDatabaseEnvironment& Environment)
 {
-	return ConfigureQueryCallbackData(Query, QueryStore, Environment, {});
+	return ConfigureQueryCallbackData(Query, QueryHandle, QueryStore, Environment, {});
 }
 
 bool UTypedElementQueryProcessorCallbackAdapterProcessorBase::ShouldAllowQueryBasedPruning(
@@ -837,11 +839,12 @@ bool UTypedElementQueryProcessorCallbackAdapterProcessorBase::ShouldAllowQueryBa
 
 bool UTypedElementQueryProcessorCallbackAdapterProcessorBase::ConfigureQueryCallbackData(
 	FTypedElementExtendedQuery& Query,
-	FTypedElementExtendedQueryStore& QueryStore, 
+	FTypedElementExtendedQueryStore::Handle QueryHandle,
+	FTypedElementExtendedQueryStore& QueryStore,
 	FTypedElementDatabaseEnvironment& Environment,
 	TArrayView<FMassEntityQuery> Subqueries)
 {
-	bool Result = Data.CommonQueryConfiguration(*this, Query, QueryStore, Environment, Subqueries);
+	bool Result = Data.CommonQueryConfiguration(*this, Query, QueryHandle, QueryStore, Environment, Subqueries);
 
 	bRequiresGameThreadExecution = Query.Description.Callback.bForceToGameThread;
 	ExecutionFlags = static_cast<int32>(EProcessorExecutionFlags::Editor); 
@@ -877,27 +880,27 @@ void UTypedElementQueryProcessorCallbackAdapterProcessorBase::Execute(FMassEntit
 }
 
 bool UTypedElementQueryProcessorCallbackAdapterProcessorWith1Subquery::ConfigureQueryCallback(
-	FTypedElementExtendedQuery& Query, FTypedElementExtendedQueryStore& QueryStore, FTypedElementDatabaseEnvironment& Environment)
+	FTypedElementExtendedQuery& Query, FTypedElementExtendedQueryStore::Handle QueryHandle, FTypedElementExtendedQueryStore& QueryStore, FTypedElementDatabaseEnvironment& Environment)
 {
-	return ConfigureQueryCallbackData(Query, QueryStore, Environment, NativeSubqueries);
+	return ConfigureQueryCallbackData(Query, QueryHandle, QueryStore, Environment, NativeSubqueries);
 }
 
 bool UTypedElementQueryProcessorCallbackAdapterProcessorWith2Subqueries::ConfigureQueryCallback(
-	FTypedElementExtendedQuery& Query, FTypedElementExtendedQueryStore& QueryStore, FTypedElementDatabaseEnvironment& Environment)
+	FTypedElementExtendedQuery& Query, FTypedElementExtendedQueryStore::Handle QueryHandle, FTypedElementExtendedQueryStore& QueryStore, FTypedElementDatabaseEnvironment& Environment)
 {
-	return ConfigureQueryCallbackData(Query, QueryStore, Environment, NativeSubqueries);
+	return ConfigureQueryCallbackData(Query, QueryHandle, QueryStore, Environment, NativeSubqueries);
 }
 
 bool UTypedElementQueryProcessorCallbackAdapterProcessorWith3Subqueries::ConfigureQueryCallback(
-	FTypedElementExtendedQuery& Query, FTypedElementExtendedQueryStore& QueryStore, FTypedElementDatabaseEnvironment& Environment)
+	FTypedElementExtendedQuery& Query, FTypedElementExtendedQueryStore::Handle QueryHandle, FTypedElementExtendedQueryStore& QueryStore, FTypedElementDatabaseEnvironment& Environment)
 {
-	return ConfigureQueryCallbackData(Query, QueryStore, Environment, NativeSubqueries);
+	return ConfigureQueryCallbackData(Query, QueryHandle, QueryStore, Environment, NativeSubqueries);
 }
 
 bool UTypedElementQueryProcessorCallbackAdapterProcessorWith4Subqueries::ConfigureQueryCallback(
-	FTypedElementExtendedQuery& Query, FTypedElementExtendedQueryStore& QueryStore, FTypedElementDatabaseEnvironment& Environment)
+	FTypedElementExtendedQuery& Query, FTypedElementExtendedQueryStore::Handle QueryHandle, FTypedElementExtendedQueryStore& QueryStore, FTypedElementDatabaseEnvironment& Environment)
 {
-	return ConfigureQueryCallbackData(Query, QueryStore, Environment, NativeSubqueries);
+	return ConfigureQueryCallbackData(Query, QueryHandle, QueryStore, Environment, NativeSubqueries);
 }
 
 
@@ -928,15 +931,18 @@ EMassObservedOperation UTypedElementQueryObserverCallbackAdapterProcessorBase::G
 }
 
 bool UTypedElementQueryObserverCallbackAdapterProcessorBase::ConfigureQueryCallback(
-	FTypedElementExtendedQuery& Query, FTypedElementExtendedQueryStore& QueryStore, FTypedElementDatabaseEnvironment& Environment)
+	FTypedElementExtendedQuery& Query, FTypedElementExtendedQueryStore::Handle QueryHandle, FTypedElementExtendedQueryStore& QueryStore, FTypedElementDatabaseEnvironment& Environment)
 {
-	return ConfigureQueryCallbackData(Query, QueryStore, Environment, {});
+	return ConfigureQueryCallbackData(Query, QueryHandle, QueryStore, Environment, {});
 }
 
-bool UTypedElementQueryObserverCallbackAdapterProcessorBase::ConfigureQueryCallbackData(FTypedElementExtendedQuery& Query,
-	FTypedElementExtendedQueryStore& QueryStore, FTypedElementDatabaseEnvironment& Environment, TArrayView<FMassEntityQuery> Subqueries)
+bool UTypedElementQueryObserverCallbackAdapterProcessorBase::ConfigureQueryCallbackData(
+	FTypedElementExtendedQuery& Query,
+	FTypedElementExtendedQueryStore::Handle QueryHandle,
+	FTypedElementExtendedQueryStore& QueryStore,
+	FTypedElementDatabaseEnvironment& Environment, TArrayView<FMassEntityQuery> Subqueries)
 {
-	bool Result = Data.CommonQueryConfiguration(*this, Query, QueryStore, Environment, Subqueries);
+	bool Result = Data.CommonQueryConfiguration(*this, Query, QueryHandle, QueryStore, Environment, Subqueries);
 
 	bRequiresGameThreadExecution = Query.Description.Callback.bForceToGameThread;
 	ExecutionFlags = static_cast<int32>(EProcessorExecutionFlags::Editor);
@@ -1001,25 +1007,25 @@ void UTypedElementQueryObserverCallbackAdapterProcessorBase::Execute(FMassEntity
 }
 
 bool UTypedElementQueryObserverCallbackAdapterProcessorWith1Subquery::ConfigureQueryCallback(
-	FTypedElementExtendedQuery& Query, FTypedElementExtendedQueryStore& QueryStore, FTypedElementDatabaseEnvironment& Environment)
+	FTypedElementExtendedQuery& Query, FTypedElementExtendedQueryStore::Handle QueryHandle, FTypedElementExtendedQueryStore& QueryStore, FTypedElementDatabaseEnvironment& Environment)
 {
-	return ConfigureQueryCallbackData(Query, QueryStore, Environment, NativeSubqueries);
+	return ConfigureQueryCallbackData(Query, QueryHandle, QueryStore, Environment, NativeSubqueries);
 }
 
 bool UTypedElementQueryObserverCallbackAdapterProcessorWith2Subqueries::ConfigureQueryCallback(
-	FTypedElementExtendedQuery& Query, FTypedElementExtendedQueryStore& QueryStore, FTypedElementDatabaseEnvironment& Environment)
+	FTypedElementExtendedQuery& Query, FTypedElementExtendedQueryStore::Handle QueryHandle, FTypedElementExtendedQueryStore& QueryStore, FTypedElementDatabaseEnvironment& Environment)
 {
-	return ConfigureQueryCallbackData(Query, QueryStore, Environment, NativeSubqueries);
+	return ConfigureQueryCallbackData(Query, QueryHandle, QueryStore, Environment, NativeSubqueries);
 }
 
 bool UTypedElementQueryObserverCallbackAdapterProcessorWith3Subqueries::ConfigureQueryCallback(
-	FTypedElementExtendedQuery& Query, FTypedElementExtendedQueryStore& QueryStore, FTypedElementDatabaseEnvironment& Environment)
+	FTypedElementExtendedQuery& Query, FTypedElementExtendedQueryStore::Handle QueryHandle, FTypedElementExtendedQueryStore& QueryStore, FTypedElementDatabaseEnvironment& Environment)
 {
-	return ConfigureQueryCallbackData(Query, QueryStore, Environment, NativeSubqueries);
+	return ConfigureQueryCallbackData(Query, QueryHandle, QueryStore, Environment, NativeSubqueries);
 }
 
 bool UTypedElementQueryObserverCallbackAdapterProcessorWith4Subqueries::ConfigureQueryCallback(
-	FTypedElementExtendedQuery& Query, FTypedElementExtendedQueryStore& QueryStore, FTypedElementDatabaseEnvironment& Environment)
+	FTypedElementExtendedQuery& Query, FTypedElementExtendedQueryStore::Handle QueryHandle, FTypedElementExtendedQueryStore& QueryStore, FTypedElementDatabaseEnvironment& Environment)
 {
-	return ConfigureQueryCallbackData(Query, QueryStore, Environment, NativeSubqueries);
+	return ConfigureQueryCallbackData(Query, QueryHandle, QueryStore, Environment, NativeSubqueries);
 }
diff --git a/Engine/Plugins/Experimental/TypedElementsDataStorage/Source/TypedElementsDataStorage/Private/Processors/TypedElementProcessorAdaptors.h b/Engine/Plugins/Experimental/TypedElementsDataStorage/Source/TypedElementsDataStorage/Private/Processors/TypedElementProcessorAdaptors.h
--- a/Engine/Plugins/Experimental/TypedElementsDataStorage/Source/TypedElementsDataStorage/Private/Processors/TypedElementProcessorAdaptors.h
+++ b/Engine/Plugins/Experimental/TypedElementsDataStorage/Source/TypedElementsDataStorage/Private/Processors/TypedElementProcessorAdaptors.h
@@ -7,6 +7,7 @@
 #include "MassExecutionContext.h"
 #include "MassObserverProcessor.h"
 #include "MassProcessor.h"
+#include "Queries/TypedElementExtendedQueryStore.h"
 #include "UObject/ObjectMacros.h"
 
 #include "TypedElementProcessorAdaptors.generated.h"
@@ -39,9 +40,10 @@ struct FTypedElementQueryProcessorData
 	explicit FTypedElementQueryProcessorData(UMassProcessor& Owner);
 
 	bool CommonQueryConfiguration(
-		UMassProcessor& InOwner, 
-		FTypedElementExtendedQuery& InQuery, 
-		FTypedElementExtendedQueryStore& InQueryStore, 
+		UMassProcessor& InOwner,
+		FTypedElementExtendedQuery& InQuery,
+		FTypedElementExtendedQueryStore::Handle InQueryHandle,
+		FTypedElementExtendedQueryStore& InQueryStore,
 		FTypedElementDatabaseEnvironment& InEnvironment,
 		TArrayView<FMassEntityQuery> Subqueries);
 	static EMassProcessingPhase MapToMassProcessingPhase(ITypedElementDataStorageInterface::EQueryTickPhase Phase);
@@ -69,7 +71,7 @@ struct FTypedElementQueryProcessorData
 		ITypedElementDataStorageInterface::FQueryDescription& Description, FMassExecutionContext& Context);
 
 	
-	FTypedElementExtendedQuery* ParentQuery{ nullptr };
+	FTypedElementExtendedQueryStore::Handle ParentQuery;
 	FTypedElementExtendedQueryStore* QueryStore{ nullptr };
 	FTypedElementDatabaseEnvironment* Environment{ nullptr };
 	FMassEntityQuery NativeQuery;
@@ -88,7 +90,8 @@ public:
 
 	FMassEntityQuery& GetQuery();
 	virtual bool ConfigureQueryCallback(
-		FTypedElementExtendedQuery& Query, 
+		FTypedElementExtendedQuery& Query,
+		FTypedElementExtendedQueryStore::Handle QueryHandle,
 		FTypedElementExtendedQueryStore& QueryStore,
 		FTypedElementDatabaseEnvironment& Environment);
 
@@ -96,8 +99,9 @@ public:
 
 protected:
 	bool ConfigureQueryCallbackData(
-		FTypedElementExtendedQuery& Query, 
-		FTypedElementExtendedQueryStore& QueryStore, 
+		FTypedElementExtendedQuery& Query,
+		FTypedElementExtendedQueryStore::Handle QueryHandle,
+		FTypedElementExtendedQueryStore& QueryStore,
 		FTypedElementDatabaseEnvironment& Environment,
 		TArrayView<FMassEntityQuery> Subqueries);
 	void ConfigureQueries() override;
@@ -132,8 +136,9 @@ class UTypedElementQueryProcessorCallbackAdapterProcessorWith1Subquery final : p
 
 public:
 	bool ConfigureQueryCallback(
-		FTypedElementExtendedQuery& Query, 
-		FTypedElementExtendedQueryStore& QueryStore, 
+		FTypedElementExtendedQuery& Query,
+		FTypedElementExtendedQueryStore::Handle QueryHandle,
+		FTypedElementExtendedQueryStore& QueryStore,
 		FTypedElementDatabaseEnvironment& Environment) override;
 
 private:
@@ -148,7 +153,8 @@ class UTypedElementQueryProcessorCallbackAdapterProcessorWith2Subqueries final :
 
 public:
 	bool ConfigureQueryCallback(
-		FTypedElementExtendedQuery& Query, 
+		FTypedElementExtendedQuery& Query,
+		FTypedElementExtendedQueryStore::Handle QueryHandle,
 		FTypedElementExtendedQueryStore& QueryStore,
 		FTypedElementDatabaseEnvironment& Environment) override;
 
@@ -164,7 +170,8 @@ class UTypedElementQueryProcessorCallbackAdapterProcessorWith3Subqueries final :
 
 public:
 	bool ConfigureQueryCallback(
-		FTypedElementExtendedQuery& Query, 
+		FTypedElementExtendedQuery& Query,
+		FTypedElementExtendedQueryStore::Handle QueryHandle,
 		FTypedElementExtendedQueryStore& QueryStore,
 		FTypedElementDatabaseEnvironment& Environment) override;
 
@@ -180,7 +187,8 @@ class UTypedElementQueryProcessorCallbackAdapterProcessorWith4Subqueries final :
 
 public:
 	bool ConfigureQueryCallback(
-		FTypedElementExtendedQuery& Query, 
+		FTypedElementExtendedQuery& Query,
+		FTypedElementExtendedQueryStore::Handle QueryHandle,
 		FTypedElementExtendedQueryStore& QueryStore,
 		FTypedElementDatabaseEnvironment& Environment) override;
 
@@ -204,11 +212,11 @@ public:
 	const UScriptStruct* GetObservedType() const;
 	EMassObservedOperation GetObservedOperation() const;
 	virtual bool ConfigureQueryCallback(
-		FTypedElementExtendedQuery& Query, FTypedElementExtendedQueryStore& QueryStore, FTypedElementDatabaseEnvironment& Environment);
+		FTypedElementExtendedQuery& Query, FTypedElementExtendedQueryStore::Handle QueryHandle, FTypedElementExtendedQueryStore& QueryStore, FTypedElementDatabaseEnvironment& Environment);
 
 protected:
-	bool ConfigureQueryCallbackData(FTypedElementExtendedQuery& Query, FTypedElementExtendedQueryStore& QueryStore,
-		FTypedElementDatabaseEnvironment& Environment, TArrayView<FMassEntityQuery> Subqueries);
+	bool ConfigureQueryCallbackData(FTypedElementExtendedQuery& Query, FTypedElementExtendedQueryStore::Handle QueryHandle,
+	                                FTypedElementExtendedQueryStore& QueryStore, FTypedElementDatabaseEnvironment& Environment, TArrayView<FMassEntityQuery> Subqueries);
 	void ConfigureQueries() override;
 	void Execute(FMassEntityManager& EntityManager, FMassExecutionContext& TargetParentQuery) override;
 
@@ -242,7 +250,8 @@ class UTypedElementQueryObserverCallbackAdapterProcessorWith1Subquery final : pu
 
 public:
 	bool ConfigureQueryCallback(
-		FTypedElementExtendedQuery& Query, 
+		FTypedElementExtendedQuery& Query,
+		FTypedElementExtendedQueryStore::Handle QueryHandle,
 		FTypedElementExtendedQueryStore& QueryStore,
 		FTypedElementDatabaseEnvironment& Environment) override;
 
@@ -258,7 +267,8 @@ class UTypedElementQueryObserverCallbackAdapterProcessorWith2Subqueries final :
 
 public:
 	bool ConfigureQueryCallback(
-		FTypedElementExtendedQuery& Query, 
+		FTypedElementExtendedQuery& Query,
+		FTypedElementExtendedQueryStore::Handle QueryHandle,
 		FTypedElementExtendedQueryStore& QueryStore,
 		FTypedElementDatabaseEnvironment& Environment) override;
 
@@ -274,7 +284,8 @@ class UTypedElementQueryObserverCallbackAdapterProcessorWith3Subqueries final :
 
 public:
 	bool ConfigureQueryCallback(
-		FTypedElementExtendedQuery& Query, 
+		FTypedElementExtendedQuery& Query,
+		FTypedElementExtendedQueryStore::Handle QueryHandle,
 		FTypedElementExtendedQueryStore& QueryStore,
 		FTypedElementDatabaseEnvironment& Environment) override;
 
@@ -290,7 +301,8 @@ class UTypedElementQueryObserverCallbackAdapterProcessorWith4Subqueries final :
 
 public:
 	bool ConfigureQueryCallback(
-		FTypedElementExtendedQuery& Query, 
+		FTypedElementExtendedQuery& Query,
+		FTypedElementExtendedQueryStore::Handle QueryHandle,
 		FTypedElementExtendedQueryStore& QueryStore,
 		FTypedElementDatabaseEnvironment& Environment) override;
 
diff --git a/Engine/Plugins/Experimental/TypedElementsDataStorage/Source/TypedElementsDataStorage/Private/Queries/TypedElementExtendedQueryStore.cpp b/Engine/Plugins/Experimental/TypedElementsDataStorage/Source/TypedElementsDataStorage/Private/Queries/TypedElementExtendedQueryStore.cpp
--- a/Engine/Plugins/Experimental/TypedElementsDataStorage/Source/TypedElementsDataStorage/Private/Queries/TypedElementExtendedQueryStore.cpp
+++ b/Engine/Plugins/Experimental/TypedElementsDataStorage/Source/TypedElementsDataStorage/Private/Queries/TypedElementExtendedQueryStore.cpp
@@ -569,7 +569,7 @@ bool FTypedElementExtendedQueryStore::SetupTickGroupDefaults(ITypedElementDataSt
 	return true;
 }
 
-bool FTypedElementExtendedQueryStore::SetupProcessors(Handle Query, FTypedElementExtendedQuery& StoredQuery, 
+bool FTypedElementExtendedQueryStore::SetupProcessors(Handle QueryHandle, FTypedElementExtendedQuery& StoredQuery, 
 	FTypedElementDatabaseEnvironment& Environment, FMassEntityManager& EntityManager, FMassProcessingPhaseManager& PhaseManager)
 {
 	using DSI = ITypedElementDataStorageInterface;
@@ -578,10 +578,10 @@ bool FTypedElementExtendedQueryStore::SetupProcessors(Handle Query, FTypedElemen
 	switch (StoredQuery.Description.Callback.Type)
 	{
 	case DSI::EQueryCallbackType::PhasePreparation:
-		RegisterPreambleQuery(StoredQuery.Description.Callback.Phase, Query);
+		RegisterPreambleQuery(StoredQuery.Description.Callback.Phase, QueryHandle);
 		break;
 	case DSI::EQueryCallbackType::PhaseFinalization:
-		RegisterPostambleQuery(StoredQuery.Description.Callback.Phase, Query);
+		RegisterPostambleQuery(StoredQuery.Description.Callback.Phase, QueryHandle);
 		break;
 	}
 
@@ -591,7 +591,7 @@ bool FTypedElementExtendedQueryStore::SetupProcessors(Handle Query, FTypedElemen
 		if (StoredQuery.Processor->IsA<UTypedElementQueryProcessorCallbackAdapterProcessorBase>())
 		{
 			if (static_cast<UTypedElementQueryProcessorCallbackAdapterProcessorBase*>(StoredQuery.Processor.Get())->
-				ConfigureQueryCallback(StoredQuery, *this, Environment))
+				ConfigureQueryCallback(StoredQuery, QueryHandle, *this, Environment))
 			{
 				PhaseManager.RegisterDynamicProcessor(*StoredQuery.Processor);
 			}
@@ -605,7 +605,7 @@ bool FTypedElementExtendedQueryStore::SetupProcessors(Handle Query, FTypedElemen
 			if (UTypedElementQueryObserverCallbackAdapterProcessorBase* Observer =
 				static_cast<UTypedElementQueryObserverCallbackAdapterProcessorBase*>(StoredQuery.Processor.Get()))
 			{
-				Observer->ConfigureQueryCallback(StoredQuery, *this, Environment);
+				Observer->ConfigureQueryCallback(StoredQuery, QueryHandle, *this, Environment);
 				EntityManager.GetObserverManager().AddObserverInstance(*Observer->GetObservedType(), Observer->GetObservedOperation(), *Observer);
 			}
 			else
diff --git a/Engine/Plugins/Experimental/TypedElementsDataStorage/Source/TypedElementsDataStorage/Private/Queries/TypedElementExtendedQueryStore.h b/Engine/Plugins/Experimental/TypedElementsDataStorage/Source/TypedElementsDataStorage/Private/Queries/TypedElementExtendedQueryStore.h
--- a/Engine/Plugins/Experimental/TypedElementsDataStorage/Source/TypedElementsDataStorage/Private/Queries/TypedElementExtendedQueryStore.h
+++ b/Engine/Plugins/Experimental/TypedElementsDataStorage/Source/TypedElementsDataStorage/Private/Queries/TypedElementExtendedQueryStore.h
@@ -136,7 +136,7 @@ private:
 	bool SetupConditions(ITypedElementDataStorageInterface::FQueryDescription& Query, FMassEntityQuery& NativeQuery);
 	bool SetupDependencies(ITypedElementDataStorageInterface::FQueryDescription& Query, FMassEntityQuery& NativeQuery);
 	bool SetupTickGroupDefaults(ITypedElementDataStorageInterface::FQueryDescription& Query);
-	bool SetupProcessors(Handle Query, FTypedElementExtendedQuery& StoredQuery, FTypedElementDatabaseEnvironment& Environment,
+	bool SetupProcessors(Handle QueryHandle, FTypedElementExtendedQuery& StoredQuery, FTypedElementDatabaseEnvironment& Environment,
 		FMassEntityManager& EntityManager, FMassProcessingPhaseManager& PhaseManager);
 
 	EMassFragmentAccess ConvertToNativeAccessType(ITypedElementDataStorageInterface::EQueryAccessType AccessType);
diff --git a/Engine/Plugins/Experimental/TypedElementsDataStorage/Source/TypedElementsDataStorage/Private/TypedElementDatabase.cpp b/Engine/Plugins/Experimental/TypedElementsDataStorage/Source/TypedElementsDataStorage/Private/TypedElementDatabase.cpp
--- a/Engine/Plugins/Experimental/TypedElementsDataStorage/Source/TypedElementsDataStorage/Private/TypedElementDatabase.cpp
+++ b/Engine/Plugins/Experimental/TypedElementsDataStorage/Source/TypedElementsDataStorage/Private/TypedElementDatabase.cpp
@@ -637,7 +637,7 @@ void UTypedElementDatabase::UnregisterTickGroup(FName GroupName, EQueryTickPhase
 TypedElementQueryHandle UTypedElementDatabase::RegisterQuery(FQueryDescription&& Query)
 {
 	return (ActiveEditorEntityManager && ActiveEditorPhaseManager)
-		? Queries.RegisterQuery(MoveTemp(Query), *Environment, *ActiveEditorEntityManager, *ActiveEditorPhaseManager).Handle
+		? Queries.RegisterQuery(MoveTemp(Query), *Environment, *ActiveEditorEntityManager, *ActiveEditorPhaseManager).Packed()
 		: TypedElementInvalidQueryHandle;
 }
 
@@ -645,17 +645,15 @@ void UTypedElementDatabase::UnregisterQuery(TypedElementQueryHandle Query)
 {
 	if (ActiveEditorPhaseManager)
 	{
-		FTypedElementExtendedQueryStore::Handle Handle;
-		Handle.Handle = Query;
-		Queries.UnregisterQuery(Handle, *ActiveEditorPhaseManager);
+		const FTypedElementExtendedQueryStore::Handle StorageHandle(Query);
+		Queries.UnregisterQuery(StorageHandle, *ActiveEditorPhaseManager);
 	}
 }
 
 const ITypedElementDataStorageInterface::FQueryDescription& UTypedElementDatabase::GetQueryDescription(TypedElementQueryHandle Query) const
 {
-	FTypedElementExtendedQueryStore::Handle Handle;
-	Handle.Handle = Query;
-	return Queries.GetQueryDescription(Handle);
+	const FTypedElementExtendedQueryStore::Handle StorageHandle(Query);
+	return Queries.GetQueryDescription(StorageHandle);
 }
 
 FName UTypedElementDatabase::GetQueryTickGroupName(EQueryTickGroups Group) const
@@ -682,9 +680,8 @@ ITypedElementDataStorageInterface::FQueryResult UTypedElementDatabase::RunQuery(
 
 	if (ActiveEditorEntityManager)
 	{
-		FTypedElementExtendedQueryStore::Handle Handle;
-		Handle.Handle = Query;
-		return Queries.RunQuery(*ActiveEditorEntityManager, Handle);
+		const FTypedElementExtendedQueryStore::Handle StorageHandle(Query);
+		return Queries.RunQuery(*ActiveEditorEntityManager, StorageHandle);
 	}
 	else
 	{
@@ -699,9 +696,8 @@ ITypedElementDataStorageInterface::FQueryResult UTypedElementDatabase::RunQuery(
 
 	if (ActiveEditorEntityManager)
 	{
-		FTypedElementExtendedQueryStore::Handle Handle;
-		Handle.Handle = Query;
-		return Queries.RunQuery(*ActiveEditorEntityManager, Handle, Callback);
+		const FTypedElementExtendedQueryStore::Handle StorageHandle(Query);
+		return Queries.RunQuery(*ActiveEditorEntityManager, StorageHandle, Callback);
 	}
 	else
 	{
diff --git a/Engine/Plugins/Experimental/TypedElementsDataStorage/Source/TypedElementsDataStorage/Private/TypedElementHandleStore.h b/Engine/Plugins/Experimental/TypedElementsDataStorage/Source/TypedElementsDataStorage/Private/TypedElementHandleStore.h
--- a/Engine/Plugins/Experimental/TypedElementsDataStorage/Source/TypedElementsDataStorage/Private/TypedElementHandleStore.h
+++ b/Engine/Plugins/Experimental/TypedElementsDataStorage/Source/TypedElementsDataStorage/Private/TypedElementHandleStore.h
@@ -17,15 +17,37 @@ public:
 		uint32 bIsAlive : 1;
 	};
 	
-	struct FHandleData
+	class Handle
 	{
-		uint32 Generation;
-		uint32 Index;
-	};
-	union Handle
-	{
-		FHandleData Data;
-		uint64 Handle;
+	public:
+		Handle() = default;
+		explicit Handle(uint64 InPackedData)
+			: PackedData(InPackedData)
+		{};
+		
+		Handle(uint32 Index, uint32 Generation)
+		{
+			PackedData = (static_cast<uint64_t>(Index) << 32) | Generation;
+		}
+
+		uint64 Packed() const
+		{
+			return PackedData;
+		}
+		
+		uint32_t Generation() const
+		{
+			return static_cast<uint32_t>(PackedData); // Unpack
+		}
+
+		uint32_t Index() const
+		{
+			return static_cast<uint32_t>(PackedData >> 32);  // Unpack
+		}
+
+		friend auto operator<=>(const Handle&, const Handle&) = default;
+	private:
+		uint64 PackedData = TNumericLimits<uint64>::Max();
 	};
 
 	using ListAliveEntriesConstCallback = TFunctionRef<void(Handle, const DataType&)>;
@@ -46,13 +68,6 @@ public:
 	void ListAliveEntries(const ListAliveEntriesConstCallback& Callback);
 	void ListAliveEntries(const ListAliveEntriesCallback& Callback);
 
-	friend bool operator==(Handle Lhs, Handle Rhs) { return Lhs.Handle == Rhs.Handle; }
-	friend bool operator!=(Handle Lhs, Handle Rhs) { return Lhs.Handle != Rhs.Handle; }
-	friend bool operator<=(Handle Lhs, Handle Rhs) { return Lhs.Handle <= Rhs.Handle; }
-	friend bool operator>=(Handle Lhs, Handle Rhs) { return Lhs.Handle >= Rhs.Handle; }
-	friend bool operator< (Handle Lhs, Handle Rhs) { return Lhs.Handle <  Rhs.Handle; }
-	friend bool operator> (Handle Lhs, Handle Rhs) { return Lhs.Handle >  Rhs.Handle; }
-
 private:
 	TArray<DataType> Data;
 	TArray<FGeneration> Generations;
@@ -87,9 +102,7 @@ auto TTypedElementHandleStore<DataType, ReservationSize>::Emplace(Args... Argume
 	DataType& Entry = Data[Index];
 	new(&Entry) DataType( Forward<Args>(Arguments)... );
 
-	Handle Result;
-	Result.Data.Index = Index;
-	Result.Data.Generation = Generations[Index].Generation;
+	Handle Result(Index, Generations[Index].Generation);
 	Generations[Index].bIsAlive = 1;
 	return Result;
 }
@@ -98,21 +111,21 @@ template<typename DataType, uint32 ReservationSize>
 DataType& TTypedElementHandleStore<DataType, ReservationSize>::Get(Handle Entry)
 {
 	checkf(IsAlive(Entry), TEXT("Attempting to retrieve a dead entry from a Typed Element Handle Store."));
-	return Data[Entry.Data.Index];
+	return Data[Entry.Index()];
 }
 
 template<typename DataType, uint32 ReservationSize>
 DataType& TTypedElementHandleStore<DataType, ReservationSize>::GetMutable(Handle Entry)
 {
 	checkf(IsAlive(Entry), TEXT("Attempting to retrieve a dead entry from a Typed Element Handle Store."));
-	return Data[Entry.Data.Index];
+	return Data[Entry.Index()];
 }
 
 template<typename DataType, uint32 ReservationSize>
 const DataType& TTypedElementHandleStore<DataType, ReservationSize>::Get(Handle Entry) const
 {
 	checkf(IsAlive(Entry), TEXT("Attempting to retrieve a dead entry from a Typed Element Handle Store."));
-	return Data[Entry.Data.Index];
+	return Data[Entry.Index()];
 }
 
 template<typename DataType, uint32 ReservationSize>
@@ -120,14 +133,14 @@ void TTypedElementHandleStore<DataType, ReservationSize>::Remove(Handle Entry)
 {
 	if (IsAlive(Entry))
 	{
-		FGeneration& Generation = Generations[Entry.Data.Index];
+		FGeneration& Generation = Generations[Entry.Index()];
 		++Generation.Generation;
 		Generation.bIsAlive = 0;
 		if constexpr (std::is_destructible_v<DataType> && !std::is_trivially_destructible_v<DataType>)
 		{
-			Data[Entry.Data.Index].~DataType();
+			Data[Entry.Index()].~DataType();
 		}
-		RecycleBin.EmplaceLast(Entry.Data.Index);
+		RecycleBin.EmplaceLast(Entry.Index());
 	}
 }
 
@@ -135,8 +148,8 @@ template<typename DataType, uint32 ReservationSize>
 bool TTypedElementHandleStore<DataType, ReservationSize>::IsAlive(Handle Entry) const
 {
 	return
-		Entry.Data.Index < static_cast<uint32>(Generations.Num()) &&
-		Generations[Entry.Data.Index].Generation == Entry.Data.Generation;
+		Entry.Index() < static_cast<uint32>(Generations.Num()) &&
+		Generations[Entry.Index()].Generation == Entry.Generation();
 }
 
 template<typename DataType, uint32 ReservationSize>
@@ -150,9 +163,7 @@ void TTypedElementHandleStore<DataType, ReservationSize>::ListAliveEntries(const
 	{
 		if (GenerationIt->bIsAlive)
 		{
-			Handle DataHandle;
-			DataHandle.Data.Index = Index;
-			DataHandle.Data.Generation = GenerationIt->Generation;
+			Handle DataHandle(Index, GenerationIt->Generation);
 			Callback(DataHandle, *EntryIt);
 		}
 
@@ -178,9 +189,7 @@ void TTypedElementHandleStore<DataType, ReservationSize>::ListAliveEntries(const
 	{
 		if (GenerationIt->bIsAlive)
 		{
-			Handle DataHandle;
-			DataHandle.Data.Index = Index;
-			DataHandle.Data.Generation = GenerationIt->Generation;
+			Handle DataHandle(Index, GenerationIt->Generation);
 			Callback(DataHandle, *EntryIt);
 		}
 
