Moved the TEDS command buffer to its own files.

diff --git a/Engine/Plugins/Experimental/TypedElementsDataStorage/Source/TypedElementsDataStorage/Private/TypedElementDatabase.cpp b/Engine/Plugins/Experimental/TypedElementsDataStorage/Source/TypedElementsDataStorage/Private/TypedElementDatabase.cpp
--- a/Engine/Plugins/Experimental/TypedElementsDataStorage/Source/TypedElementsDataStorage/Private/TypedElementDatabase.cpp
+++ b/Engine/Plugins/Experimental/TypedElementsDataStorage/Source/TypedElementsDataStorage/Private/TypedElementDatabase.cpp
@@ -118,7 +118,8 @@ void UTypedElementDatabase::OnPreMassTick(float DeltaTime)
 	OnUpdateDelegate.Broadcast();
 	// Process pending commands after other systems have had a chance to update. Other systems may have executed work needed
 	// to complete pending work.
-	ProcessPendingCommands();
+	FTypedElementDatabaseCommandBuffer::ProcessBuffer(DeferredCommands, *Environment);
+	DeferredCommands.Reset();
 
 	Environment->GetScratchBuffer().NextFrame();
 }
@@ -280,29 +281,25 @@ void UTypedElementDatabase::RemoveRow(TypedElementRowHandle Row)
 
 bool UTypedElementDatabase::IsRowAvailable(TypedElementRowHandle Row) const
 {
-	return ActiveEditorEntityManager ? ActiveEditorEntityManager->IsEntityValid(FMassEntityHandle::FromNumber(Row)) : false;
+	return ActiveEditorEntityManager ? FTypedElementDatabaseCommandBuffer::Execute_IsRowAvailable(*ActiveEditorEntityManager, Row) : false;
 }
 
 bool UTypedElementDatabase::HasRowBeenAssigned(TypedElementRowHandle Row) const
 {
-	return ActiveEditorEntityManager ? ActiveEditorEntityManager->IsEntityActive(FMassEntityHandle::FromNumber(Row)) : false;
+	return ActiveEditorEntityManager ? FTypedElementDatabaseCommandBuffer::Execute_HasRowBeenAssigned(*ActiveEditorEntityManager, Row) : false;
 }
 
 bool UTypedElementDatabase::AddColumn(TypedElementRowHandle Row, const UScriptStruct* ColumnType)
 {
-	if (ColumnType && ActiveEditorEntityManager)
+	if (ColumnType && ActiveEditorEntityManager && HasRowBeenAssigned(Row))
 	{
-		if (HasRowBeenAssigned(Row))
-		{
-			ExecuteAddColumnCommand(Row, ColumnType);
-		}
-		else
-		{
-			AddPendingCommand(Row, FAddColumnCommand{ .ColumnType = ColumnType });
-		}
-		return true;
+		FTypedElementDatabaseCommandBuffer::Execute_AddColumnCommand(*ActiveEditorEntityManager, Row, ColumnType);
 	}
-	return false;
+	else
+	{
+		FTypedElementDatabaseCommandBuffer::Queue_AddColumnCommand(DeferredCommands, Row, ColumnType);
+	}
+	return true;
 }
 
 bool UTypedElementDatabase::AddColumn(TypedElementRowHandle Row, FTopLevelAssetPath ColumnName)
@@ -314,16 +311,13 @@ bool UTypedElementDatabase::AddColumn(TypedElementRowHandle Row, FTopLevelAssetP
 
 void UTypedElementDatabase::RemoveColumn(TypedElementRowHandle Row, const UScriptStruct* ColumnType)
 {
-	if (ActiveEditorEntityManager)
+	if (ActiveEditorEntityManager && HasRowBeenAssigned(Row))
 	{
-		if (HasRowBeenAssigned(Row))
-		{
-			ExecuteRemoveColumnCommand(Row, ColumnType);
-		}
-		else
-		{
-			AddPendingCommand(Row, FRemoveColumnCommand{ .ColumnType = ColumnType });
-		}
+		FTypedElementDatabaseCommandBuffer::Execute_RemoveColumnCommand(*ActiveEditorEntityManager, Row, ColumnType);
+	}
+	else
+	{
+		FTypedElementDatabaseCommandBuffer::Queue_RemoveColumnCommand(DeferredCommands, Row, ColumnType);
 	}
 }
 
@@ -424,9 +418,9 @@ ColumnDataResult UTypedElementDatabase::GetColumnData(TypedElementRowHandle Row,
 
 bool UTypedElementDatabase::AddColumns(TypedElementRowHandle Row, TConstArrayView<const UScriptStruct*> Columns)
 {
-	FMassEntityHandle Entity = FMassEntityHandle::FromNumber(Row);
 	if (ActiveEditorEntityManager)
 	{
+		FMassEntityHandle Entity = FMassEntityHandle::FromNumber(Row);
 		FMassArchetypeHandle Archetype = ActiveEditorEntityManager->GetArchetypeForEntity(Entity);
 
 		FMassFragmentBitSet FragmentsToAdd;
@@ -435,11 +429,11 @@ bool UTypedElementDatabase::AddColumns(TypedElementRowHandle Row, TConstArrayVie
 		{
 			if (ActiveEditorEntityManager->IsEntityActive(Entity))
 			{
-				ExecuteAddColumnsCommand(Row, FragmentsToAdd, TagsToAdd);
+				FTypedElementDatabaseCommandBuffer::Execute_AddColumnsCommand(*ActiveEditorEntityManager, Row, FragmentsToAdd, TagsToAdd);
 			}
 			else
 			{
-				AddPendingCommand(Row, FAddColumnsCommand{ .FragmentsToAdd = MoveTemp(FragmentsToAdd), .TagsToAdd = MoveTemp(TagsToAdd) });
+				FTypedElementDatabaseCommandBuffer::Queue_AddColumnsCommand(DeferredCommands, Row, FragmentsToAdd, TagsToAdd);
 			}
 			return true;
 		}
@@ -460,12 +454,11 @@ void UTypedElementDatabase::RemoveColumns(TypedElementRowHandle Row, TConstArray
 		{
 			if (ActiveEditorEntityManager->IsEntityActive(Entity))
 			{
-				ExecuteRemoveColumnsCommand(Row, FragmentsToRemove, TagsToRemove);
+				FTypedElementDatabaseCommandBuffer::Execute_RemoveColumnsCommand(*ActiveEditorEntityManager, Row, FragmentsToRemove, TagsToRemove);
 			}
 			else
 			{
-				AddPendingCommand(Row, 
-					FRemoveColumnsCommand{ .FragmentsToRemove = MoveTemp(FragmentsToRemove), .TagsToRemove = MoveTemp(TagsToRemove) });
+				FTypedElementDatabaseCommandBuffer::Queue_RemoveColumnsCommand(DeferredCommands, Row, FragmentsToRemove, TagsToRemove);
 			}
 		}
 	}
@@ -764,97 +757,6 @@ bool UTypedElementDatabase::ColumnsToBitSets(TConstArrayView<const UScriptStruct
 	return bResult;
 }
 
-template<typename T>
-void UTypedElementDatabase::AddPendingCommand(TypedElementRowHandle Row, T&& Args)
-{
-	FCommand Command;
-	Command.Row = Row;
-	Command.Data.Emplace<T>(Forward<T>(Args));
-	PendingCommands.Add(MoveTemp(Command));
-}
-
-void UTypedElementDatabase::ProcessPendingCommands()
-{
-	PendingCommands.StableSort(
-		[](const FCommand& Lhs, const FCommand& Rhs)
-		{
-			return Lhs.Row < Rhs.Row;
-		});
-	
-	struct FProcessor
-	{
-		UTypedElementDatabase* This;
-		TypedElementRowHandle Row;
-		void operator()(const FAddColumnCommand& Command)		{ This->ExecuteAddColumnCommand(Row, Command.ColumnType.Get()); }
-		void operator()(const FAddColumnsCommand& Command)		{ This->ExecuteAddColumnsCommand(Row, Command.FragmentsToAdd, Command.TagsToAdd); }
-		void operator()(const FRemoveColumnCommand& Command)	{ This->ExecuteRemoveColumnCommand(Row, Command.ColumnType.Get()); }
-		void operator()(const FRemoveColumnsCommand& Command)	{ This->ExecuteRemoveColumnsCommand(Row, Command.FragmentsToRemove, Command.TagsToRemove); }
-	};
-	FProcessor Processor;
-	Processor.This = this;
-
-	for (FCommand& Command : PendingCommands)
-	{
-		if (HasRowBeenAssigned(Command.Row))
-		{
-			Processor.Row = Command.Row;
-			Visit(Processor, Command.Data);
-		}
-	}
-	PendingCommands.Reset();
-}
-
-void UTypedElementDatabase::ExecuteAddColumnCommand(TypedElementRowHandle Row, const UScriptStruct* ColumnType)
-{
-	if (ColumnType)
-	{
-		FMassEntityHandle Entity = FMassEntityHandle::FromNumber(Row);
-		if (ColumnType->IsChildOf(FMassTag::StaticStruct()))
-		{
-			ActiveEditorEntityManager->AddTagToEntity(Entity, ColumnType);
-		}
-		else if (ColumnType->IsChildOf(FMassFragment::StaticStruct()))
-		{
-			FStructView Column = ActiveEditorEntityManager->GetFragmentDataStruct(Entity, ColumnType);
-			// Only add if not already added to avoid asserts from Mass.
-			if (!Column.IsValid())
-			{
-				ActiveEditorEntityManager->AddFragmentToEntity(Entity, ColumnType);
-			}
-		}
-	}
-}
-
-void UTypedElementDatabase::ExecuteRemoveColumnCommand(TypedElementRowHandle Row, const UScriptStruct* ColumnType)
-{
-	if (ColumnType)
-	{
-		FMassEntityHandle Entity = FMassEntityHandle::FromNumber(Row);
-		if (ColumnType->IsChildOf(FMassTag::StaticStruct()))
-		{
-			ActiveEditorEntityManager->RemoveTagFromEntity(Entity, ColumnType);
-		}
-		else if (ColumnType->IsChildOf(FMassFragment::StaticStruct()))
-		{
-			ActiveEditorEntityManager->RemoveFragmentFromEntity(Entity, ColumnType);
-		}
-	}
-}
-
-void UTypedElementDatabase::ExecuteAddColumnsCommand(TypedElementRowHandle Row, FMassFragmentBitSet FragmentsToAdd, FMassTagBitSet TagsToAdd)
-{
-	FMassArchetypeCompositionDescriptor AddComposition(
-		MoveTemp(FragmentsToAdd), MoveTemp(TagsToAdd), FMassChunkFragmentBitSet(), FMassSharedFragmentBitSet());
-	ActiveEditorEntityManager->AddCompositionToEntity_GetDelta(FMassEntityHandle::FromNumber(Row), AddComposition);
-}
-
-void UTypedElementDatabase::ExecuteRemoveColumnsCommand(TypedElementRowHandle Row, FMassFragmentBitSet FragmentsToRemove, FMassTagBitSet TagsToRemove)
-{
-	FMassArchetypeCompositionDescriptor RemoveComposition(
-		MoveTemp(FragmentsToRemove), MoveTemp(TagsToRemove), FMassChunkFragmentBitSet(), FMassSharedFragmentBitSet());
-	ActiveEditorEntityManager->RemoveCompositionFromEntity(FMassEntityHandle::FromNumber(Row), RemoveComposition);
-}
-
 void UTypedElementDatabase::PreparePhase(EQueryTickPhase Phase, float DeltaTime)
 {
 	if (ActiveEditorEntityManager)
diff --git a/Engine/Plugins/Experimental/TypedElementsDataStorage/Source/TypedElementsDataStorage/Private/TypedElementDatabase.h b/Engine/Plugins/Experimental/TypedElementsDataStorage/Source/TypedElementsDataStorage/Private/TypedElementDatabase.h
--- a/Engine/Plugins/Experimental/TypedElementsDataStorage/Source/TypedElementsDataStorage/Private/TypedElementDatabase.h
+++ b/Engine/Plugins/Experimental/TypedElementsDataStorage/Source/TypedElementsDataStorage/Private/TypedElementDatabase.h
@@ -8,6 +8,7 @@
 #include "Misc/TVariant.h"
 #include "Queries/TypedElementExtendedQueryStore.h"
 #include "Templates/SharedPointer.h"
+#include "TypedElementDatabaseCommandBuffer.h"
 #include "TypedElementDatabaseEnvironment.h"
 #include "UObject/ObjectMacros.h"
 #include "UObject/StrongObjectPtr.h"
@@ -98,44 +99,10 @@ public:
 
 	void DebugPrintQueryCallbacks(FOutputDevice& Output);
 
-private:	
-	struct FAddColumnCommand
-	{
-		TWeakObjectPtr<const UScriptStruct> ColumnType;
-	};
-	struct FAddColumnsCommand
-	{
-		FMassFragmentBitSet FragmentsToAdd;
-		FMassTagBitSet TagsToAdd;
-	};
-	struct FRemoveColumnCommand
-	{
-		TWeakObjectPtr<const UScriptStruct> ColumnType;
-	};
-	struct FRemoveColumnsCommand
-	{
-		FMassFragmentBitSet FragmentsToRemove;
-		FMassTagBitSet TagsToRemove;
-	};
-	using CommandData = TVariant<FAddColumnCommand, FAddColumnsCommand, FRemoveColumnCommand, FRemoveColumnsCommand>;
-
-	struct FCommand
-	{
-		TypedElementRowHandle Row;
-		CommandData Data;
-	};
-	
+private:
 	/** Converts a set of column types into Mass specific fragment and tag bit sets. Returns true if any values were added. */
 	static bool ColumnsToBitSets(TConstArrayView<const UScriptStruct*> Columns, FMassFragmentBitSet& Fragments, FMassTagBitSet& Tags);
 
-	template<typename T>
-	void AddPendingCommand(TypedElementRowHandle Row, T&& Args);
-	void ProcessPendingCommands();
-	void ExecuteAddColumnCommand(TypedElementRowHandle Row, const UScriptStruct* ColumnType);
-	void ExecuteAddColumnsCommand(TypedElementRowHandle Row, FMassFragmentBitSet FragmentsToAdd, FMassTagBitSet TagsToAdd);
-	void ExecuteRemoveColumnCommand(TypedElementRowHandle Row, const UScriptStruct* ColumnType);
-	void ExecuteRemoveColumnsCommand(TypedElementRowHandle Row, FMassFragmentBitSet FragmentsToRemove, FMassTagBitSet TagsToRemove);
-	
 	void PreparePhase(EQueryTickPhase Phase, float DeltaTime);
 	void FinalizePhase(EQueryTickPhase Phase, float DeltaTime);
 	void Reset();
@@ -147,6 +114,8 @@ private:
 
 	TUniquePtr<FTypedElementDatabaseEnvironment> Environment;
 	FTypedElementExtendedQueryStore Queries;
+
+	FTypedElementDatabaseCommandBuffer::CommandBuffer DeferredCommands;
 	
 	FTypedElementOnDataStorageUpdate OnUpdateDelegate;
 	FDelegateHandle OnPreMassTickHandle;
@@ -154,6 +123,4 @@ private:
 
 	TSharedPtr<FMassEntityManager> ActiveEditorEntityManager;
 	TSharedPtr<FMassProcessingPhaseManager> ActiveEditorPhaseManager;
-
-	TArray<FCommand> PendingCommands;
 };
\ No newline at end of file
diff --git a/Engine/Plugins/Experimental/TypedElementsDataStorage/Source/TypedElementsDataStorage/Private/TypedElementDatabaseCommandBuffer.cpp b/Engine/Plugins/Experimental/TypedElementsDataStorage/Source/TypedElementsDataStorage/Private/TypedElementDatabaseCommandBuffer.cpp
new file mode 100644
--- /dev/null
+++ b/Engine/Plugins/Experimental/TypedElementsDataStorage/Source/TypedElementsDataStorage/Private/TypedElementDatabaseCommandBuffer.cpp
@@ -0,0 +1,181 @@
+// Copyright Epic Games, Inc. All Rights Reserved.
+
+#include "TypedElementDatabaseCommandBuffer.h"
+
+#include "MassEntityManager.h"
+#include "TypedElementDatabaseEnvironment.h"
+
+//
+// Queue section
+//
+
+void FTypedElementDatabaseCommandBuffer::Queue_AddColumnCommand(
+	CommandBuffer& Buffer,
+	TypedElementDataStorage::RowHandle Row,
+	const UScriptStruct* ColumnType)
+{
+	AddCommand(Buffer, Row, FAddColumnCommand
+		{ 
+			.ColumnType = ColumnType 
+		});
+}
+
+void FTypedElementDatabaseCommandBuffer::Queue_AddColumnsCommand(
+	CommandBuffer& Buffer,
+	TypedElementDataStorage::RowHandle Row,
+	FMassFragmentBitSet FragmentsToAdd,
+	FMassTagBitSet TagsToAdd)
+{
+	AddCommand(Buffer, Row, FAddColumnsCommand
+		{ 
+			.FragmentsToAdd = MoveTemp(FragmentsToAdd), 
+			.TagsToAdd = MoveTemp(TagsToAdd) 
+		});
+}
+
+void FTypedElementDatabaseCommandBuffer::Queue_RemoveColumnCommand(
+	CommandBuffer& Buffer,
+	TypedElementDataStorage::RowHandle Row,
+	const UScriptStruct* ColumnType)
+{
+	AddCommand(Buffer, Row, FRemoveColumnCommand
+		{ 
+			.ColumnType = ColumnType 
+		});
+}
+
+void FTypedElementDatabaseCommandBuffer::Queue_RemoveColumnsCommand(
+	CommandBuffer& Buffer,
+	TypedElementDataStorage::RowHandle Row,
+	FMassFragmentBitSet FragmentsToRemove,
+	FMassTagBitSet TagsToRemove)
+{
+	AddCommand(Buffer, Row, FRemoveColumnsCommand
+		{ 
+			.FragmentsToRemove = MoveTemp(FragmentsToRemove), 
+			.TagsToRemove = MoveTemp(TagsToRemove) 
+		});
+}
+
+//
+// Execute section
+//
+
+bool FTypedElementDatabaseCommandBuffer::Execute_IsRowAvailable(const FMassEntityManager& MassEntityManager, TypedElementDataStorage::RowHandle Row)
+{
+	return MassEntityManager.IsEntityValid(FMassEntityHandle::FromNumber(Row));
+}
+
+bool FTypedElementDatabaseCommandBuffer::Execute_HasRowBeenAssigned(const FMassEntityManager& MassEntityManager, TypedElementDataStorage::RowHandle Row)
+{
+	return MassEntityManager.IsEntityActive(FMassEntityHandle::FromNumber(Row));
+}
+
+void FTypedElementDatabaseCommandBuffer::Execute_AddColumnCommand(
+	FMassEntityManager& MassEntityManager,
+	TypedElementDataStorage::RowHandle Row,
+	const UScriptStruct* ColumnType)
+{
+	if (ColumnType)
+	{
+		FMassEntityHandle Entity = FMassEntityHandle::FromNumber(Row);
+		if (ColumnType->IsChildOf(FMassTag::StaticStruct()))
+		{
+			MassEntityManager.AddTagToEntity(Entity, ColumnType);
+		}
+		else if (ColumnType->IsChildOf(FMassFragment::StaticStruct()))
+		{
+			FStructView Column = MassEntityManager.GetFragmentDataStruct(Entity, ColumnType);
+			// Only add if not already added to avoid asserts from Mass.
+			if (!Column.IsValid())
+			{
+				MassEntityManager.AddFragmentToEntity(Entity, ColumnType);
+			}
+		}
+	}
+}
+
+void FTypedElementDatabaseCommandBuffer::Execute_AddColumnsCommand(
+	FMassEntityManager& MassEntityManager,
+	TypedElementDataStorage::RowHandle Row,
+	FMassFragmentBitSet FragmentsToAdd,
+	FMassTagBitSet TagsToAdd)
+{
+	FMassArchetypeCompositionDescriptor AddComposition(
+		MoveTemp(FragmentsToAdd), MoveTemp(TagsToAdd), FMassChunkFragmentBitSet(), FMassSharedFragmentBitSet());
+	MassEntityManager.AddCompositionToEntity_GetDelta(FMassEntityHandle::FromNumber(Row), AddComposition);
+}
+
+void FTypedElementDatabaseCommandBuffer::Execute_RemoveColumnCommand(
+	FMassEntityManager& MassEntityManager,
+	TypedElementDataStorage::RowHandle Row,
+	const UScriptStruct* ColumnType)
+{
+	if (ColumnType)
+	{
+		FMassEntityHandle Entity = FMassEntityHandle::FromNumber(Row);
+		if (ColumnType->IsChildOf(FMassTag::StaticStruct()))
+		{
+			MassEntityManager.RemoveTagFromEntity(Entity, ColumnType);
+		}
+		else if (ColumnType->IsChildOf(FMassFragment::StaticStruct()))
+		{
+			MassEntityManager.RemoveFragmentFromEntity(Entity, ColumnType);
+		}
+	}
+}
+
+void FTypedElementDatabaseCommandBuffer::Execute_RemoveColumnsCommand(
+	FMassEntityManager& MassEntityManager,
+	TypedElementDataStorage::RowHandle Row,
+	FMassFragmentBitSet FragmentsToRemove,
+	FMassTagBitSet TagsToRemove)
+{
+	FMassArchetypeCompositionDescriptor RemoveComposition(
+		MoveTemp(FragmentsToRemove), MoveTemp(TagsToRemove), FMassChunkFragmentBitSet(), FMassSharedFragmentBitSet());
+	MassEntityManager.RemoveCompositionFromEntity(FMassEntityHandle::FromNumber(Row), RemoveComposition);
+}
+
+void FTypedElementDatabaseCommandBuffer::ProcessBuffer(CommandBuffer& Buffer, FTypedElementDatabaseEnvironment& Environment)
+{
+	using namespace TypedElementDataStorage;
+
+	Buffer.StableSort(
+		[](const FCommand& Lhs, const FCommand& Rhs)
+		{
+			return Lhs.Row < Rhs.Row;
+		});
+
+	struct FProcessor
+	{
+		FMassEntityManager& EntityManager;
+		RowHandle Row;
+		void operator()(const FAddColumnCommand& Command) { Execute_AddColumnCommand(EntityManager, Row, Command.ColumnType.Get()); }
+		void operator()(const FAddColumnsCommand& Command) { Execute_AddColumnsCommand(EntityManager, Row, Command.FragmentsToAdd, Command.TagsToAdd); }
+		void operator()(const FRemoveColumnCommand& Command) { Execute_RemoveColumnCommand(EntityManager, Row, Command.ColumnType.Get()); }
+		void operator()(const FRemoveColumnsCommand& Command) { Execute_RemoveColumnsCommand(EntityManager, Row, Command.FragmentsToRemove, Command.TagsToRemove); }
+	};
+	
+	FProcessor Processor{ .EntityManager = Environment.GetMassEntityManager() };
+	for (FCommand& Command : Buffer)
+	{
+		if (Execute_HasRowBeenAssigned(Environment.GetMassEntityManager(), Command.Row))
+		{
+			Processor.Row = Command.Row;
+			Visit(Processor, Command.Data);
+		}
+	}
+}
+
+// 
+// misc
+//
+
+template<typename T>
+void FTypedElementDatabaseCommandBuffer::AddCommand(CommandBuffer& Buffer, TypedElementDataStorage::RowHandle Row, T&& Args)
+{
+	FCommand Command;
+	Command.Row = Row;
+	Command.Data.Emplace<T>(Forward<T>(Args));
+	Buffer.Add(MoveTemp(Command));
+}
diff --git a/Engine/Plugins/Experimental/TypedElementsDataStorage/Source/TypedElementsDataStorage/Private/TypedElementDatabaseCommandBuffer.h b/Engine/Plugins/Experimental/TypedElementsDataStorage/Source/TypedElementsDataStorage/Private/TypedElementDatabaseCommandBuffer.h
new file mode 100644
--- /dev/null
+++ b/Engine/Plugins/Experimental/TypedElementsDataStorage/Source/TypedElementsDataStorage/Private/TypedElementDatabaseCommandBuffer.h
@@ -0,0 +1,101 @@
+// Copyright Epic Games, Inc. All Rights Reserved.
+
+#pragma once
+
+#include "Elements/Common/TypedElementHandles.h"
+#include "MassArchetypeTypes.h"
+#include "Misc/TVariant.h"
+#include "UObject/Class.h"
+#include "UObject/WeakObjectPtrTemplates.h"
+
+class FTypedElementDatabaseEnvironment;
+struct FMassEntityManager;
+struct FMassProcessingPhaseManager;
+
+class FTypedElementDatabaseCommandBuffer final
+{
+private:
+	struct FAddColumnCommand
+	{
+		TWeakObjectPtr<const UScriptStruct> ColumnType;
+	};
+	struct FAddColumnsCommand
+	{
+		FMassFragmentBitSet FragmentsToAdd;
+		FMassTagBitSet TagsToAdd;
+	};
+	struct FRemoveColumnCommand
+	{
+		TWeakObjectPtr<const UScriptStruct> ColumnType;
+	};
+	struct FRemoveColumnsCommand
+	{
+		FMassFragmentBitSet FragmentsToRemove;
+		FMassTagBitSet TagsToRemove;
+	};
+	using CommandData = TVariant<FAddColumnCommand, FAddColumnsCommand, FRemoveColumnCommand, FRemoveColumnsCommand>;
+
+	struct FCommand
+	{
+		TypedElementDataStorage::RowHandle Row;
+		CommandData Data;
+	};
+
+public:
+	using CommandBuffer = TArray<FCommand>;
+
+	/**
+	 * @section Queue_* functions create a command that can be stored for later execution in the provided command buffer.
+	 * Typically the Execute_* counter part will be called for execution.
+	 */
+	static void Queue_AddColumnCommand(
+		CommandBuffer& Buffer,
+		TypedElementDataStorage::RowHandle Row,
+		const UScriptStruct* ColumnType);
+	static void Queue_AddColumnsCommand(
+		CommandBuffer& Buffer,
+		TypedElementDataStorage::RowHandle Row,
+		FMassFragmentBitSet FragmentsToAdd,
+		FMassTagBitSet TagsToAdd);
+	static void Queue_RemoveColumnCommand(
+		CommandBuffer& Buffer,
+		TypedElementDataStorage::RowHandle Row,
+		const UScriptStruct* ColumnType);
+	static void Queue_RemoveColumnsCommand(
+		CommandBuffer& Buffer,
+		TypedElementDataStorage::RowHandle Row,
+		FMassFragmentBitSet FragmentsToRemove,
+		FMassTagBitSet TagsToRemove);
+
+	/**
+	 * @section Execute_* functions directly execute a command with limited validation checks.
+	 */
+	
+	static bool Execute_IsRowAvailable(const FMassEntityManager& MassEntityManager, TypedElementDataStorage::RowHandle Row);
+	static bool Execute_HasRowBeenAssigned(const FMassEntityManager& MassEntityManager, TypedElementDataStorage::RowHandle Row);
+
+	static void Execute_AddColumnCommand(
+		FMassEntityManager& MassEntityManager, 
+		TypedElementDataStorage::RowHandle Row,
+		const UScriptStruct* ColumnType);
+	static void Execute_AddColumnsCommand(
+		FMassEntityManager& MassEntityManager, 
+		TypedElementDataStorage::RowHandle Row,
+		FMassFragmentBitSet FragmentsToAdd, 
+		FMassTagBitSet TagsToAdd);
+	static void Execute_RemoveColumnCommand(
+		FMassEntityManager& MassEntityManager, 
+		TypedElementDataStorage::RowHandle Row,
+		const UScriptStruct* ColumnType);
+	static void Execute_RemoveColumnsCommand(
+		FMassEntityManager& MassEntityManager, 
+		TypedElementDataStorage::RowHandle Row,
+		FMassFragmentBitSet FragmentsToRemove, 
+		FMassTagBitSet TagsToRemove);
+
+	static void ProcessBuffer(CommandBuffer& Buffer, FTypedElementDatabaseEnvironment& Environment);
+
+private:
+	template<typename T>
+	static void AddCommand(CommandBuffer& Buffer, TypedElementDataStorage::RowHandle Row, T&& Args);
+};
\ No newline at end of file
