Fixed an incorrect unregistration in TEDS Compat of OnObjectModified delegate.

#nojira

diff --git a/Engine/Plugins/Experimental/TypedElementsDataStorage/Source/TypedElementsDataStorage/Private/TypedElementDatabaseCompatibility.cpp b/Engine/Plugins/Experimental/TypedElementsDataStorage/Source/TypedElementsDataStorage/Private/TypedElementDatabaseCompatibility.cpp
--- a/Engine/Plugins/Experimental/TypedElementsDataStorage/Source/TypedElementsDataStorage/Private/TypedElementDatabaseCompatibility.cpp
+++ b/Engine/Plugins/Experimental/TypedElementsDataStorage/Source/TypedElementsDataStorage/Private/TypedElementDatabaseCompatibility.cpp
@@ -33,7 +33,7 @@ void UTypedElementDatabaseCompatibility::Initialize(ITypedElementDataStorageInte
 
 void UTypedElementDatabaseCompatibility::Deinitialize()
 {
-	FCoreUObjectDelegates::OnObjectPropertyChanged.Remove(ObjectModifiedDelegateHandle);
+	FCoreUObjectDelegates::OnObjectModified.Remove(ObjectModifiedDelegateHandle);
 	FCoreUObjectDelegates::OnObjectPropertyChanged.Remove(PostEditChangePropertyDelegateHandle);
 	
 	Reset();
