Fixed assert when unparenting an entity.

TEDS favors simplicity so supports queueing up the deletion of columns that may not be on a row. This reduces complexity in the query callbacks as these do not need to first verify if a column is actually present. An example of this in action is the removal of FTypedElementParentColumn and FTypedElementUnresolvedParentColumn in one call when it's discovered that there's no longer a parent. Only one of these two columns will ever be active on a row.
Mass however takes a stricter approach and doesn't allow the removal of a fragment/tag from an archetype unless it's there. To bridge the two, TEDS will behind the scenes check for this case and filter out any fragments/tags that aren't on the archetype. To achieve this Mass had to get a small update to be able to access the required information.

diff --git a/Engine/Plugins/Experimental/TypedElementsDataStorage/Source/TypedElementsDataStorage/Private/Processors/TypedElementProcessorAdaptors.cpp b/Engine/Plugins/Experimental/TypedElementsDataStorage/Source/TypedElementsDataStorage/Private/Processors/TypedElementProcessorAdaptors.cpp
--- a/Engine/Plugins/Experimental/TypedElementsDataStorage/Source/TypedElementsDataStorage/Private/Processors/TypedElementProcessorAdaptors.cpp
+++ b/Engine/Plugins/Experimental/TypedElementsDataStorage/Source/TypedElementsDataStorage/Private/Processors/TypedElementProcessorAdaptors.cpp
@@ -150,6 +150,31 @@ private:
 struct FMassContextForwarder final : public FMassContextCommon<ITypedElementDataStorageInterface::IQueryContext>
 {
 private:
+	void TedsColumnsToMassDescriptorIfActiveTable(
+		FMassArchetypeCompositionDescriptor& Descriptor, 
+		TConstArrayView<const UScriptStruct*> ColumnTypes)
+	{
+		for (const UScriptStruct* ColumnType : ColumnTypes)
+		{
+			if (ColumnType->IsChildOf(FMassTag::StaticStruct()))
+			{
+				if (Context.DoesArchetypeHaveTag(*ColumnType))
+				{
+					Descriptor.Tags.Add(*ColumnType);
+				}
+			}
+			else
+			{
+				checkf(ColumnType->IsChildOf(FMassFragment::StaticStruct()),
+					TEXT("Given struct type is not a valid fragment or tag type."));
+				if (Context.DoesArchetypeHaveFragment(*ColumnType))
+				{
+					Descriptor.Fragments.Add(*ColumnType);
+				}
+			}
+		}
+	}
+
 	void TedsColumnsToMassDescriptor(FMassArchetypeCompositionDescriptor& Descriptor, TConstArrayView<const UScriptStruct*> ColumnTypes)
 	{
 		for (const UScriptStruct* ColumnType : ColumnTypes)
@@ -470,7 +495,7 @@ public:
 		};
 
 		FRemovedColumns* RemovedColumns = Environment.GetScratchBuffer().Emplace<FRemovedColumns>();
-		TedsColumnsToMassDescriptor(RemovedColumns->RemoveDescriptor, ColumnTypes);
+		TedsColumnsToMassDescriptorIfActiveTable(RemovedColumns->RemoveDescriptor, ColumnTypes);
 		RemovedColumns->Entity = FMassEntityHandle::FromNumber(Row);
 
 		Context.Defer().PushCommand<FMassDeferredAddCommand>(
@@ -491,7 +516,7 @@ public:
 
 		FTypedElementDatabaseScratchBuffer& ScratchBuffer = Environment.GetScratchBuffer();
 		FRemovedColumns* RemovedColumns = ScratchBuffer.Emplace<FRemovedColumns>();
-		TedsColumnsToMassDescriptor(RemovedColumns->RemoveDescriptor, ColumnTypes);
+		TedsColumnsToMassDescriptorIfActiveTable(RemovedColumns->RemoveDescriptor, ColumnTypes);
 
 		FMassEntityHandle* Entities = ScratchBuffer.EmplaceArray<FMassEntityHandle>(Rows.Num());
 		RemovedColumns->Entities = Entities;
diff --git a/Engine/Plugins/Runtime/MassEntity/Source/MassEntity/Public/MassExecutionContext.h b/Engine/Plugins/Runtime/MassEntity/Source/MassEntity/Public/MassExecutionContext.h
--- a/Engine/Plugins/Runtime/MassEntity/Source/MassEntity/Public/MassExecutionContext.h
+++ b/Engine/Plugins/Runtime/MassEntity/Source/MassEntity/Public/MassExecutionContext.h
@@ -119,10 +119,31 @@ public:
 		return EntityListView[Index];
 	}
 
+	bool DoesArchetypeHaveFragment(const UScriptStruct& FragmentType) const
+	{
+		return FragmentViews.FindByPredicate(
+			[&FragmentType](const FFragmentView& Element) 
+			{ 
+				return Element.Requirement.StructType == &FragmentType; 
+			}) != nullptr;
+	}
+
+	template<typename T>
+	bool DoesArchetypeHaveFragment() const
+	{
+		static_assert(TIsDerivedFrom<T, FMassFragment>::IsDerived, "Given struct is not of a valid fragment type.");
+		return DoesArchetypeHaveFragment(T::StaticStruct());
+	}
+
+	bool DoesArchetypeHaveTag(const UScriptStruct& TagType) const
+	{
+		return CurrentArchetypesTagBitSet.Contains(TagType);
+	}
+
 	template<typename T>
 	bool DoesArchetypeHaveTag() const
 	{
-		static_assert(TIsDerivedFrom<T, FMassTag>::IsDerived, "Given struct is not of a valid fragment type.");
+		static_assert(TIsDerivedFrom<T, FMassTag>::IsDerived, "Given struct is not of a valid tag type.");
 		return CurrentArchetypesTagBitSet.Contains<T>();
 	}
 
