Add the ability to test for column presence from a TEDS query

This makes it possible to test a given row for tags

diff --git a/Engine/Plugins/Experimental/TypedElementsDataStorage/Source/TypedElementsDataStorage/Private/Processors/TypedElementProcessorAdaptors.cpp b/Engine/Plugins/Experimental/TypedElementsDataStorage/Source/TypedElementsDataStorage/Private/Processors/TypedElementProcessorAdaptors.cpp
--- a/Engine/Plugins/Experimental/TypedElementsDataStorage/Source/TypedElementsDataStorage/Private/Processors/TypedElementProcessorAdaptors.cpp
+++ b/Engine/Plugins/Experimental/TypedElementsDataStorage/Source/TypedElementsDataStorage/Private/Processors/TypedElementProcessorAdaptors.cpp
@@ -67,6 +67,21 @@ struct FMassContextCommon : public T
 		}
 	}
 
+	bool HasColumn(const UScriptStruct* ColumnType) const override
+	{
+		if (ColumnType->IsChildOf(FMassTag::StaticStruct()))
+		{
+			return Context.DoesArchetypeHaveTag(*ColumnType);
+		}
+		if (ColumnType->IsChildOf(FMassFragment::StaticStruct()))
+		{
+			return Context.DoesArchetypeHaveFragment(*ColumnType);
+		}
+		const bool bIsTagOrFragment = false;
+		checkf(bIsTagOrFragment, TEXT("Attempting to check for a column type that is not a column or tag."))
+		return false;
+	}
+
 protected:
 	explicit FMassContextCommon(FMassExecutionContext& InContext)
 		: Context(InContext)
@@ -142,6 +157,25 @@ struct FMassSingleRowSubqueryContextForwarder final : public ITypedElementDataSt
 		return TConstArrayView<TypedElementDataStorage::RowHandle>(&RowHandle, 1);
 	}
 
+	bool HasColumn(const UScriptStruct* ColumnType) const override
+	{
+		FMassArchetypeHandle Archetype = EntityManager.GetArchetypeForEntity(FMassEntityHandle::FromNumber(RowHandle));
+		FMassArchetypeCompositionDescriptor Composition = EntityManager.GetArchetypeComposition(Archetype);
+		// Tags
+		if (ColumnType->IsChildOf(FMassTag::StaticStruct()))
+		{
+			return Composition.Tags.Contains(*ColumnType);
+		}
+		// Columns
+		if (ColumnType->IsChildOf(FMassFragment::StaticStruct()))
+		{
+			return Composition.Fragments.Contains(*ColumnType);
+		}
+		const bool bIsTagOrFragment = false;
+		checkf(bIsTagOrFragment, TEXT("Attempting to check for a column type that is not a column or tag."))
+		return false;
+	}
+
 private:
 	FMassEntityManager& EntityManager;
 	TypedElementDataStorage::RowHandle RowHandle;
diff --git a/Engine/Source/Runtime/TypedElementFramework/Public/Elements/Framework/TypedElementQueryBuilder.h b/Engine/Source/Runtime/TypedElementFramework/Public/Elements/Framework/TypedElementQueryBuilder.h
--- a/Engine/Source/Runtime/TypedElementFramework/Public/Elements/Framework/TypedElementQueryBuilder.h
+++ b/Engine/Source/Runtime/TypedElementFramework/Public/Elements/Framework/TypedElementQueryBuilder.h
@@ -365,6 +365,8 @@ namespace TypedElementQueryBuilder
 		inline void GetColumnsUnguarded(int32 TypeCount, char** RetrievedAddresses, const TWeakObjectPtr<const UScriptStruct>* ColumnTypes,
 			const ITypedElementDataStorageInterface::EQueryAccessType* AccessTypes) override;
 
+		inline bool HasColumn(const UScriptStruct* ColumnType) const override;
+		
 		inline UObject* GetMutableDependency(const UClass* DependencyClass) override;
 		inline const UObject* GetDependency(const UClass* DependencyClass) override;
 		inline void GetDependencies(TArrayView<UObject*> RetrievedAddresses, TConstArrayView<TWeakObjectPtr<const UClass>> SubsystemTypes,
diff --git a/Engine/Source/Runtime/TypedElementFramework/Public/Elements/Framework/TypedElementQueryBuilder.inl b/Engine/Source/Runtime/TypedElementFramework/Public/Elements/Framework/TypedElementQueryBuilder.inl
--- a/Engine/Source/Runtime/TypedElementFramework/Public/Elements/Framework/TypedElementQueryBuilder.inl
+++ b/Engine/Source/Runtime/TypedElementFramework/Public/Elements/Framework/TypedElementQueryBuilder.inl
@@ -159,6 +159,11 @@ namespace TypedElementQueryBuilder
 	{
 		ParentContext.GetColumnsUnguarded(TypeCount, RetrievedAddresses, ColumnTypes, AccessTypes);
 	}
+	
+	bool FQueryContextForwarder::HasColumn(const UScriptStruct* ColumnType) const
+	{
+		return ParentContext.HasColumn(ColumnType);
+	}
 
 	UObject* FQueryContextForwarder::GetMutableDependency(const UClass* DependencyClass)
 	{
diff --git a/Engine/Source/Runtime/TypedElementFramework/Public/Elements/Interfaces/TypedElementQueryStorageInterfaces.h b/Engine/Source/Runtime/TypedElementFramework/Public/Elements/Interfaces/TypedElementQueryStorageInterfaces.h
--- a/Engine/Source/Runtime/TypedElementFramework/Public/Elements/Interfaces/TypedElementQueryStorageInterfaces.h
+++ b/Engine/Source/Runtime/TypedElementFramework/Public/Elements/Interfaces/TypedElementQueryStorageInterfaces.h
@@ -55,13 +55,16 @@ namespace TypedElementDataStorage
 		 */
 		virtual void GetColumnsUnguarded(int32 TypeCount, char** RetrievedAddresses, const TWeakObjectPtr<const UScriptStruct>* ColumnTypes,
 			const EQueryAccessType* AccessTypes) = 0;
-
+		/* Return whether a column matches the requested type or not. */
+		virtual bool HasColumn(const UScriptStruct* ColumnType) const = 0;
 		// Utility functions
 
 		template<typename Column>
 		const Column* GetColumn() const;
 		template<typename Column>
 		Column* GetMutableColumn();
+		template<typename Column>
+		bool HasColumn() const;
 	};
 
 	/**
@@ -221,6 +224,12 @@ namespace TypedElementDataStorage
 		return reinterpret_cast<Column*>(GetMutableColumn(Column::StaticStruct()));
 	}
 
+	template <typename Column>
+	bool ICommonQueryContext::HasColumn() const
+	{
+		return HasColumn(Column::StaticStruct());
+	}
+
 	template<typename ColumnType>
 	ColumnType& IQueryContext::AddColumn(RowHandle Row, ColumnType&& Column)
 	{
