(Un)registering objects with TEDS Compat is now a transactional operation.

Registering and unregistering objects with TEDS Compat was not registered in the transaction buffer as a result registering as part of an transaction would cause the registered objects to remain in TEDS after an undo. This was fixed by adding an explicit command.

diff --git a/Engine/Plugins/Experimental/TypedElementsDataStorage/Source/TypedElementsDataStorage/Private/TypedElementDatabaseCompatibility.cpp b/Engine/Plugins/Experimental/TypedElementsDataStorage/Source/TypedElementsDataStorage/Private/TypedElementDatabaseCompatibility.cpp
--- a/Engine/Plugins/Experimental/TypedElementsDataStorage/Source/TypedElementsDataStorage/Private/TypedElementDatabaseCompatibility.cpp
+++ b/Engine/Plugins/Experimental/TypedElementsDataStorage/Source/TypedElementsDataStorage/Private/TypedElementDatabaseCompatibility.cpp
@@ -4,6 +4,7 @@
 
 #include "Algo/Unique.h"
 #include "Editor.h"
+#include "Editor/TransBuffer.h"
 #include "Elements/Columns/TypedElementCompatibilityColumns.h"
 #include "Elements/Columns/TypedElementLabelColumns.h"
 #include "Elements/Columns/TypedElementMiscColumns.h"
@@ -86,31 +87,18 @@ TypedElementRowHandle UTypedElementDatabaseCompatibility::AddCompatibleObjectExp
 
 TypedElementRowHandle UTypedElementDatabaseCompatibility::AddCompatibleObjectExplicit(UObject* Object, TypedElementTableHandle Table)
 {
-	using namespace TypedElementDataStorage;
-
-#if TEDS_SEPARATE_ACTOR_REGISTRATION
-	if (Object->IsA<AActor>())
+	if (ensureMsgf(Storage, TEXT("Trying to add a UObject to Typed Element's Data Storage before the storage is available.")) &&
+		ShouldAddObject(Object))
 	{
-		return AddCompatibleObjectExplicit(static_cast<AActor*>(Object));
+		if (GUndo)
+		{
+			GUndo->StoreUndo(Object, MakeUnique<FRegistrationCommandChange>(Table, this));
+		}
+		return AddCompatibleObjectExplicitNoTransaction(Object, Table);
 	}
 	else
-#endif
 	{
-		if (ensureMsgf(Storage, TEXT("Trying to add a UObject to Typed Element's Data Storage before the storage is available.")) && 
-			ShouldAddObject(Object))
-		{
-			TypedElementRowHandle ReservedRow = Storage->ReserveRow();
-			Storage->IndexRow(GenerateIndexHash(Object), ReservedRow);
-			
-			PendingRegistration<TWeakObjectPtr<UObject>>& Pending = UObjectsPendingRegistration.FindOrAdd(Table);
-			Pending.Add(ReservedRow, Object);
-			
-			return ReservedRow;
-		}
-		else
-		{
-			return TypedElementInvalidRowHandle;
-		}
+		return TypedElementInvalidRowHandle;
 	}
 }
 
@@ -354,6 +342,28 @@ bool UTypedElementDatabaseCompatibility::ShouldAddObject(const UObject* Object)
 	return Include;
 }
 
+TypedElementRowHandle UTypedElementDatabaseCompatibility::AddCompatibleObjectExplicitNoTransaction(UObject* Object, TypedElementTableHandle Table)
+{
+	using namespace TypedElementDataStorage;
+
+#if TEDS_SEPARATE_ACTOR_REGISTRATION
+	if (Object->IsA<AActor>())
+	{
+		return AddCompatibleObjectExplicit(static_cast<AActor*>(Object));
+	}
+	else
+#endif
+	{
+		TypedElementRowHandle ReservedRow = Storage->ReserveRow();
+		Storage->IndexRow(GenerateIndexHash(Object), ReservedRow);
+
+		PendingRegistration<TWeakObjectPtr<UObject>>& Pending = UObjectsPendingRegistration.FindOrAdd(Table);
+		Pending.Add(ReservedRow, Object);
+
+		return ReservedRow;
+	}
+}
+
 TypedElementRowHandle UTypedElementDatabaseCompatibility::DealiasObject(const UObject* Object) const
 {
 	for (const ObjectToRowDealiaser& Dealiaser : ObjectToRowDialiasers)
@@ -554,7 +564,6 @@ void UTypedElementDatabaseCompatibility::TickPendingUObjectRegistration()
 					Storage->AddOrGetColumn<FTypedElementClassTypeInfoColumn>(Row, FTypedElementClassTypeInfoColumn{ .TypeInfo = Object->GetClass() });
 					// Make sure the new row is tagged for update.
 					Storage->AddColumn<FTypedElementSyncFromWorldTag>(Row);
-					
 					OnObjectAdded(Object.Get(), Object->GetClass(), Row);
 				});
 		}
@@ -672,3 +681,31 @@ void UTypedElementDatabaseCompatibility::OnPreObjectRemoved(const void* Object,
 		Callback(Object, TypeInfo, Row);
 	}
 }
+
+
+
+//
+// UTypedElementDatabaseCompatibility::FRegistrationCommandChange
+//
+
+UTypedElementDatabaseCompatibility::FRegistrationCommandChange::FRegistrationCommandChange(
+	TypedElementDataStorage::TableHandle InTable, UTypedElementDatabaseCompatibility* InCompatibilityLayer)
+	: CompatibilityLayer(InCompatibilityLayer)
+	, Table(InTable)
+{
+}
+
+void UTypedElementDatabaseCompatibility::FRegistrationCommandChange::Apply(UObject* Object)
+{
+	CompatibilityLayer->AddCompatibleObjectExplicitNoTransaction(Object, Table);
+}
+
+void UTypedElementDatabaseCompatibility::FRegistrationCommandChange::Revert(UObject* Object)
+{
+	CompatibilityLayer->RemoveCompatibleObject(Object);
+}
+
+FString UTypedElementDatabaseCompatibility::FRegistrationCommandChange::ToString() const
+{
+	return TEXT("Typed Element Data Storage Compatibility - Registration");
+}
diff --git a/Engine/Plugins/Experimental/TypedElementsDataStorage/Source/TypedElementsDataStorage/Private/TypedElementDatabaseCompatibility.h b/Engine/Plugins/Experimental/TypedElementsDataStorage/Source/TypedElementsDataStorage/Private/TypedElementDatabaseCompatibility.h
--- a/Engine/Plugins/Experimental/TypedElementsDataStorage/Source/TypedElementsDataStorage/Private/TypedElementDatabaseCompatibility.h
+++ b/Engine/Plugins/Experimental/TypedElementsDataStorage/Source/TypedElementsDataStorage/Private/TypedElementDatabaseCompatibility.h
@@ -4,9 +4,10 @@
 
 #include "Containers/Array.h"
 #include "Containers/Map.h"
+#include "Compatibility/TypedElementObjectReinstancingManager.h"
 #include "Elements/Interfaces/TypedElementDataStorageInterface.h"
 #include "Elements/Interfaces/TypedElementDataStorageCompatibilityInterface.h"
-#include "Compatibility/TypedElementObjectReinstancingManager.h"
+#include "Misc/Change.h"
 #include "UObject/ObjectMacros.h"
 #include "UObject/WeakObjectPtr.h"
 
@@ -61,11 +62,27 @@ public:
 	TypedElementRowHandle FindRowWithCompatibleObjectExplicit(const void* Object) const override;
 
 private:
+	class FRegistrationCommandChange final : public FCommandChange
+	{
+	public:
+		FRegistrationCommandChange(TypedElementDataStorage::TableHandle InTable,
+			UTypedElementDatabaseCompatibility* InCompatibilityLayer);
+
+		void Apply(UObject* Object) override;
+		void Revert(UObject* Object) override;
+		FString ToString() const override;
+
+	private:
+		UTypedElementDatabaseCompatibility* CompatibilityLayer{ nullptr };
+		TypedElementDataStorage::TableHandle Table{ TypedElementDataStorage::InvalidTableHandle };
+	};
+
 	void Prepare();
 	void Reset();
 	void CreateStandardArchetypes();
 	
 	bool ShouldAddObject(const UObject* Object) const;
+	TypedElementRowHandle AddCompatibleObjectExplicitNoTransaction(UObject* Object, TypedElementTableHandle Table);
 	TypedElementRowHandle DealiasObject(const UObject* Object) const;
 
 	void Tick();
