Fixed undoing the removal of an object not causing to reregister that object with TEDS.

Undoing the deletion of an object will not always trigger that object to go through its full initialization path and as a result may not end up being registered with TEDS again. This change mirrors an earlier change that allowed correctly undoing a creation, however some more data was needed as the table isn't know at destruction. The table handle is now stored with the object in TEDS so it can be retrieved during removal and stored in an transaction command. This assumes that after restoring the object can be fully restored based on the information stored in the object. Additional information such as selection may need to be stored through the memento system, but this integration was left as future work on order to address the immediate bug.

diff --git a/Engine/Plugins/Experimental/TypedElementsDataStorage/Source/TypedElementsDataStorage/Private/TypedElementDatabaseCompatibility.cpp b/Engine/Plugins/Experimental/TypedElementsDataStorage/Source/TypedElementsDataStorage/Private/TypedElementDatabaseCompatibility.cpp
--- a/Engine/Plugins/Experimental/TypedElementsDataStorage/Source/TypedElementsDataStorage/Private/TypedElementDatabaseCompatibility.cpp
+++ b/Engine/Plugins/Experimental/TypedElementsDataStorage/Source/TypedElementsDataStorage/Private/TypedElementDatabaseCompatibility.cpp
@@ -87,19 +87,10 @@ TypedElementRowHandle UTypedElementDatabaseCompatibility::AddCompatibleObjectExp
 
 TypedElementRowHandle UTypedElementDatabaseCompatibility::AddCompatibleObjectExplicit(UObject* Object, TypedElementTableHandle Table)
 {
-	if (ensureMsgf(Storage, TEXT("Trying to add a UObject to Typed Element's Data Storage before the storage is available.")) &&
-		ShouldAddObject(Object))
-	{
-		if (GUndo)
-		{
-			GUndo->StoreUndo(Object, MakeUnique<FRegistrationCommandChange>(Table, this));
-		}
-		return AddCompatibleObjectExplicitNoTransaction(Object, Table);
-	}
-	else
-	{
-		return TypedElementInvalidRowHandle;
-	}
+	bool bCanAddObject = 
+		ensureMsgf(Storage, TEXT("Trying to add a UObject to Typed Element's Data Storage before the storage is available.")) &&
+		ShouldAddObject(Object);
+	return bCanAddObject ? AddCompatibleObjectExplicitTransactionable<true>(Object, Table) : TypedElementDataStorage::InvalidRowHandle;
 }
 
 TypedElementRowHandle UTypedElementDatabaseCompatibility::AddCompatibleObjectExplicit(AActor* Actor)
@@ -164,18 +155,7 @@ void UTypedElementDatabaseCompatibility::RemoveCompatibleObjectExplicit(UObject*
 	else
 #else
 	{
-		checkf(Storage, TEXT("Removing compatible objects is not supported before Typed Element's Database compatibility manager has been initialized."));
-		IndexHash Hash = GenerateIndexHash(Object);
-		RowHandle Row = Storage->FindIndexedRow(Hash);
-		if (Storage->IsRowAvailable(Row))
-		{
-			const FTypedElementClassTypeInfoColumn* TypeInfoColumn = Storage->GetColumn<FTypedElementClassTypeInfoColumn>(Row);
-			if (Storage->HasRowBeenAssigned(Row) && ensureMsgf(TypeInfoColumn, TEXT("Missing type information for removed UObject at ptr 0x%p [%s]"), Object, *Object->GetName()))
-			{
-				OnPreObjectRemoved(Object, TypeInfoColumn->TypeInfo.Get(), Row);
-			}
-			Storage->RemoveRow(Row);
-		}
+		RemoveCompatibleObjectExplicitTransactionable<true>(Object);
 	}
 #endif
 }
@@ -306,6 +286,7 @@ void UTypedElementDatabaseCompatibility::CreateStandardArchetypes()
 {
 	StandardActorTable = Storage->RegisterTable(TTypedElementColumnTypeList<
 			FMassActorFragment, FTypedElementUObjectColumn, FTypedElementClassTypeInfoColumn,
+			FTypedElementObjectSourceTableColumn,
 			FTypedElementLabelColumn, FTypedElementLabelHashColumn,
 			FTypedElementSyncFromWorldTag>(), 
 		FName("Editor_StandardActorTable"));
@@ -316,11 +297,13 @@ void UTypedElementDatabaseCompatibility::CreateStandardArchetypes()
 
 	StandardUObjectTable = Storage->RegisterTable(TTypedElementColumnTypeList<
 			FTypedElementUObjectColumn, FTypedElementClassTypeInfoColumn,
+			FTypedElementObjectSourceTableColumn,
 			FTypedElementSyncFromWorldTag>(), 
 		FName("Editor_StandardUObjectTable"));
 
 	StandardExternalObjectTable = Storage->RegisterTable(TTypedElementColumnTypeList<
 			FTypedElementExternalObjectColumn, FTypedElementScriptStructTypeInfoColumn,
+			FTypedElementObjectSourceTableColumn,
 			FTypedElementSyncFromWorldTag>(), 
 		FName("Editor_StandardExternalObjectTable"));
 }
@@ -342,7 +325,8 @@ bool UTypedElementDatabaseCompatibility::ShouldAddObject(const UObject* Object)
 	return Include;
 }
 
-TypedElementRowHandle UTypedElementDatabaseCompatibility::AddCompatibleObjectExplicitNoTransaction(UObject* Object, TypedElementTableHandle Table)
+template<bool bEnableTransactions>
+TypedElementRowHandle UTypedElementDatabaseCompatibility::AddCompatibleObjectExplicitTransactionable(UObject* Object, TypedElementTableHandle Table)
 {
 	using namespace TypedElementDataStorage;
 
@@ -360,10 +344,61 @@ TypedElementRowHandle UTypedElementDatabaseCompatibility::AddCompatibleObjectExp
 		PendingRegistration<TWeakObjectPtr<UObject>>& Pending = UObjectsPendingRegistration.FindOrAdd(Table);
 		Pending.Add(ReservedRow, Object);
 
+		if constexpr (bEnableTransactions)
+		{
+			if (GUndo)
+			{
+				GUndo->StoreUndo(Object, MakeUnique<FRegistrationCommandChange>(Table, this));
+			}
+		}
+
 		return ReservedRow;
 	}
 }
 
+template<bool bEnableTransactions>
+void UTypedElementDatabaseCompatibility::RemoveCompatibleObjectExplicitTransactionable(UObject* Object)
+{
+	using namespace TypedElementDataStorage;
+
+	checkf(Storage, 
+		TEXT("Removing compatible objects is not supported before Typed Element's Database compatibility manager has been initialized."));
+	IndexHash Hash = GenerateIndexHash(Object);
+	RowHandle Row = Storage->FindIndexedRow(Hash);
+	if (Storage->IsRowAvailable(Row))
+	{
+		const FTypedElementClassTypeInfoColumn* TypeInfoColumn = Storage->GetColumn<FTypedElementClassTypeInfoColumn>(Row);
+		if (Storage->HasRowBeenAssigned(Row) && 
+			ensureMsgf(TypeInfoColumn, TEXT("Missing type information for removed UObject at ptr 0x%p [%s]"), Object, *Object->GetName()))
+		{
+			OnPreObjectRemoved(Object, TypeInfoColumn->TypeInfo.Get(), Row);
+
+			if constexpr (bEnableTransactions)
+			{
+				if (GUndo)
+				{
+					FTypedElementObjectSourceTableColumn* Table = Storage->GetColumn<FTypedElementObjectSourceTableColumn>(Row);
+					if (ensureMsgf(Table,
+						TEXT("An object is removed from the Typed Element's Database compatibility manager that didn't contain a source table column.")))
+					{
+						// Reverting the deletion of a row relies on the original table used to create the row rather
+						// than the mutated variant that may have been used at the time of destruction. The reasoning is
+						// that using the original table will trigger the same update that the mutated versions will call
+						// so both will result in the same mirrored data in the data storage after processing. The benefit
+						// of taking the original data is that less information has to be stored in the memento table thus
+						// reducing the memory requirements though at the cost of more table mutations during reconstruction.
+						// This does however put more emphasis on storing any auxiliary data that can't be reconstructed from
+						// the object, e.g. the selection column, through the memento system.
+						GUndo->StoreUndo(Object, MakeUnique<FDeregistrationCommandChange>(Table->SourceTable, this));
+					}
+				}
+			}
+		}
+
+		Storage->RemoveRow(Row);
+	}
+}
+
 TypedElementRowHandle UTypedElementDatabaseCompatibility::DealiasObject(const UObject* Object) const
 {
 	for (const ObjectToRowDealiaser& Dealiaser : ObjectToRowDialiasers)
@@ -551,7 +586,7 @@ void UTypedElementDatabaseCompatibility::TickPendingUObjectRegistration()
 		for (auto It = UObjectsPendingRegistration.CreateIterator(); It; ++It)
 		{
 			It->Value.ProcessEntries(*Storage, It->Key,
-				[this](TypedElementRowHandle Row, const TWeakObjectPtr<UObject>& Object)
+				[Table = It->Key, this](TypedElementRowHandle Row, const TWeakObjectPtr<UObject>& Object)
 				{
 #if !TEDS_SEPARATE_ACTOR_REGISTRATION
 					if (AActor* Actor = Cast<AActor>(Object))
@@ -562,6 +597,7 @@ void UTypedElementDatabaseCompatibility::TickPendingUObjectRegistration()
 #endif
 					Storage->AddOrGetColumn<FTypedElementUObjectColumn>(Row, FTypedElementUObjectColumn{ .Object = Object });
 					Storage->AddOrGetColumn<FTypedElementClassTypeInfoColumn>(Row, FTypedElementClassTypeInfoColumn{ .TypeInfo = Object->GetClass() });
+					Storage->AddOrGetColumn<FTypedElementObjectSourceTableColumn>(Row, FTypedElementObjectSourceTableColumn{ .SourceTable = Table });
 					// Make sure the new row is tagged for update.
 					Storage->AddColumn<FTypedElementSyncFromWorldTag>(Row);
 					OnObjectAdded(Object.Get(), Object->GetClass(), Row);
@@ -590,10 +626,11 @@ void UTypedElementDatabaseCompatibility::TickPendingExternalObjectRegistration()
 		for (auto It = ExternalObjectsPendingRegistration.CreateIterator(); It; ++It)
 		{
 			It->Value.ProcessEntries(*Storage, It->Key, 
-				[this](TypedElementRowHandle Row, const ExternalObjectRegistration& Object)
+				[Table = It->Key, this](TypedElementRowHandle Row, const ExternalObjectRegistration& Object)
 				{
 					Storage->AddOrGetColumn<FTypedElementExternalObjectColumn>(Row, FTypedElementExternalObjectColumn{ .Object = Object.Object });
 					Storage->AddOrGetColumn<FTypedElementScriptStructTypeInfoColumn>(Row, FTypedElementScriptStructTypeInfoColumn{ .TypeInfo = Object.TypeInfo });
+					Storage->AddOrGetColumn<FTypedElementObjectSourceTableColumn>(Row, FTypedElementObjectSourceTableColumn{ .SourceTable = Table });
 					// Make sure the new row is tagged for update.
 					Storage->AddColumn<FTypedElementSyncFromWorldTag>(Row);
 
@@ -697,15 +734,42 @@ UTypedElementDatabaseCompatibility::FRegistrationCommandChange::FRegistrationCom
 
 void UTypedElementDatabaseCompatibility::FRegistrationCommandChange::Apply(UObject* Object)
 {
-	CompatibilityLayer->AddCompatibleObjectExplicitNoTransaction(Object, Table);
+	CompatibilityLayer->AddCompatibleObjectExplicitTransactionable<false>(Object, Table);
 }
 
 void UTypedElementDatabaseCompatibility::FRegistrationCommandChange::Revert(UObject* Object)
 {
-	CompatibilityLayer->RemoveCompatibleObject(Object);
+	CompatibilityLayer->RemoveCompatibleObjectExplicitTransactionable<false>(Object);
 }
 
 FString UTypedElementDatabaseCompatibility::FRegistrationCommandChange::ToString() const
 {
 	return TEXT("Typed Element Data Storage Compatibility - Registration");
 }
+
+
+//
+// UTypedElementDatabaseCompatibility::FDeregistrationCommandChange
+//
+
+UTypedElementDatabaseCompatibility::FDeregistrationCommandChange::FDeregistrationCommandChange(
+	TypedElementDataStorage::TableHandle InTable, UTypedElementDatabaseCompatibility* InCompatibilityLayer)
+	: CompatibilityLayer(InCompatibilityLayer)
+	, Table(InTable)
+{
+}
+
+void UTypedElementDatabaseCompatibility::FDeregistrationCommandChange::Apply(UObject* Object)
+{
+	CompatibilityLayer->RemoveCompatibleObjectExplicitTransactionable<false>(Object);
+}
+
+void UTypedElementDatabaseCompatibility::FDeregistrationCommandChange::Revert(UObject* Object)
+{
+	CompatibilityLayer->AddCompatibleObjectExplicitTransactionable<false>(Object, Table);
+}
+
+FString UTypedElementDatabaseCompatibility::FDeregistrationCommandChange::ToString() const
+{
+	return TEXT("Typed Element Data Storage Compatibility - Deregistration");
+}
diff --git a/Engine/Plugins/Experimental/TypedElementsDataStorage/Source/TypedElementsDataStorage/Private/TypedElementDatabaseCompatibility.h b/Engine/Plugins/Experimental/TypedElementsDataStorage/Source/TypedElementsDataStorage/Private/TypedElementDatabaseCompatibility.h
--- a/Engine/Plugins/Experimental/TypedElementsDataStorage/Source/TypedElementsDataStorage/Private/TypedElementDatabaseCompatibility.h
+++ b/Engine/Plugins/Experimental/TypedElementsDataStorage/Source/TypedElementsDataStorage/Private/TypedElementDatabaseCompatibility.h
@@ -77,12 +77,30 @@ private:
 		TypedElementDataStorage::TableHandle Table{ TypedElementDataStorage::InvalidTableHandle };
 	};
 
+	class FDeregistrationCommandChange final : public FCommandChange
+	{
+	public:
+		FDeregistrationCommandChange(TypedElementDataStorage::TableHandle InTable,
+			UTypedElementDatabaseCompatibility* InCompatibilityLayer);
+
+		void Apply(UObject* Object) override;
+		void Revert(UObject* Object) override;
+		FString ToString() const override;
+
+	private:
+		UTypedElementDatabaseCompatibility* CompatibilityLayer{ nullptr };
+		TypedElementDataStorage::TableHandle Table{ TypedElementDataStorage::InvalidTableHandle };
+	};
+
 	void Prepare();
 	void Reset();
 	void CreateStandardArchetypes();
 	
 	bool ShouldAddObject(const UObject* Object) const;
-	TypedElementRowHandle AddCompatibleObjectExplicitNoTransaction(UObject* Object, TypedElementTableHandle Table);
+	template<bool bEnableTransactions>
+	TypedElementRowHandle AddCompatibleObjectExplicitTransactionable(UObject* Object, TypedElementTableHandle Table);
+	template<bool bEnableTransactions>
+	void RemoveCompatibleObjectExplicitTransactionable(UObject* Object);
 	TypedElementRowHandle DealiasObject(const UObject* Object) const;
 
 	void Tick();
diff --git a/Engine/Source/Runtime/TypedElementFramework/Public/Elements/Columns/TypedElementCompatibilityColumns.h b/Engine/Source/Runtime/TypedElementFramework/Public/Elements/Columns/TypedElementCompatibilityColumns.h
--- a/Engine/Source/Runtime/TypedElementFramework/Public/Elements/Columns/TypedElementCompatibilityColumns.h
+++ b/Engine/Source/Runtime/TypedElementFramework/Public/Elements/Columns/TypedElementCompatibilityColumns.h
@@ -2,6 +2,7 @@
 
 #pragma once
 
+#include "Elements/Common/TypedElementHandles.h"
 #include "Elements/Interfaces/TypedElementDataStorageInterface.h"
 #include "UObject/ObjectMacros.h"
 #include "UObject/WeakObjectPtrTemplates.h"
@@ -32,4 +33,15 @@ struct FTypedElementExternalObjectColumn final : public FTypedElementDataStorage
 	GENERATED_BODY()
 
 	void* Object;
+};
+
+/**
+ * Source table used to register the external object with.
+ */
+USTRUCT(meta = (DisplayName = "Object source table"))
+struct FTypedElementObjectSourceTableColumn final : public FTypedElementDataStorageColumn
+{
+	GENERATED_BODY()
+
+	TypedElementDataStorage::TableHandle SourceTable;
 };
\ No newline at end of file
