Fix running destructor for contents of TTypedElementHandleStore twice.
Make TTypedElementHandleStore immovable and uncopyable, implement destructor manually to only destroy live elements.

diff --git a/Engine/Plugins/Experimental/TypedElementsDataStorage/Source/TypedElementsDataStorage/Private/TypedElementHandleStore.h b/Engine/Plugins/Experimental/TypedElementsDataStorage/Source/TypedElementsDataStorage/Private/TypedElementHandleStore.h
--- a/Engine/Plugins/Experimental/TypedElementsDataStorage/Source/TypedElementsDataStorage/Private/TypedElementHandleStore.h
+++ b/Engine/Plugins/Experimental/TypedElementsDataStorage/Source/TypedElementsDataStorage/Private/TypedElementHandleStore.h
@@ -6,6 +6,7 @@
 #include "Containers/Array.h"
 #include "Containers/Deque.h"
 #include "Templates/Function.h"
+#include "Templates/MemoryOps.h"
 
 template<typename DataType, uint32 ReservationSize = 64>
 class TTypedElementHandleStore
@@ -53,6 +54,13 @@ public:
 	using ListAliveEntriesConstCallback = TFunctionRef<void(Handle, const DataType&)>;
 	using ListAliveEntriesCallback = TFunctionRef<void(Handle, DataType&)>;
 	
+	TTypedElementHandleStore() = default;
+	~TTypedElementHandleStore();
+	TTypedElementHandleStore(const TTypedElementHandleStore&) = delete;
+	TTypedElementHandleStore(TTypedElementHandleStore&&) = delete;
+	TTypedElementHandleStore& operator=(const TTypedElementHandleStore&) = delete;
+	TTypedElementHandleStore& operator=(TTypedElementHandleStore&&) = delete;
+	
 	template<typename... Args>
 	Handle Emplace(Args... Arguments);
 	
@@ -77,6 +85,28 @@ private:
 
 
 // Implementations
+template<typename DataType, uint32 ReservationSize>
+TTypedElementHandleStore<DataType, ReservationSize>::~TTypedElementHandleStore()
+{
+	if constexpr (std::is_destructible_v<DataType> && !std::is_trivially_destructible_v<DataType>)
+	{
+		int32 Count = Data.Num();
+		const DataType* EntryIt = Data.GetData();
+		const FGeneration* GenerationIt = Generations.GetData();
+		
+		for (int32 Index = 0; Index < Count; ++Index)
+		{
+			if (GenerationIt->bIsAlive)
+			{
+				DestructItem(EntryIt);
+			}
+
+			++EntryIt;
+			++GenerationIt;
+		}
+	}
+	Data.SetNumUnsafeInternal(0); // Make sure array destructor does not run destructor on items again
+}
 
 template<typename DataType, uint32 ReservationSize>
 template<typename... Args>
@@ -138,7 +168,7 @@ void TTypedElementHandleStore<DataType, ReservationSize>::Remove(Handle Entry)
 		Generation.bIsAlive = 0;
 		if constexpr (std::is_destructible_v<DataType> && !std::is_trivially_destructible_v<DataType>)
 		{
-			Data[Entry.Index()].~DataType();
+			DestructItem(&Data[Entry.Index()]);
 		}
 		RecycleBin.EmplaceLast(Entry.Index());
 	}
