Add CVar to enable TypedElementDataStorage
* Disabled by default

diff --git a/Engine/Plugins/Experimental/TypedElementsDataStorage/Source/TypedElementsDataStorage/Private/TypedElementsDataStorage.cpp b/Engine/Plugins/Experimental/TypedElementsDataStorage/Source/TypedElementsDataStorage/Private/TypedElementsDataStorage.cpp
--- a/Engine/Plugins/Experimental/TypedElementsDataStorage/Source/TypedElementsDataStorage/Private/TypedElementsDataStorage.cpp
+++ b/Engine/Plugins/Experimental/TypedElementsDataStorage/Source/TypedElementsDataStorage/Private/TypedElementsDataStorage.cpp
@@ -10,10 +10,22 @@
 #include "TypedElementDatabase.h"
 #include "TypedElementDatabaseCompatibility.h"
 #include "TypedElementDatabaseUI.h"
+#include "Elements/Common/TypedElementDataStorageLog.h"
 #include "UObject/UObjectGlobals.h"
 
 #define LOCTEXT_NAMESPACE "FTypedElementsDataStorageModule"
 
+namespace Private
+{
+	static bool GEnableTEDS = false;
+	static FAutoConsoleVariableRef CVarEnableTEDS(
+		TEXT("TEDS.Enable"),
+		GEnableTEDS,
+		TEXT("Enable TypedElementDataStorage"),
+		ECVF_ReadOnly
+		);
+} // namespace Private
+
 // MASS uses CDO in a few places, making it a difficult to consistently register Type Element's Columns and Tags
 // as they may have not been set up to impersonate MASS' Fragments and Tags yet. There are currently no longer
 // any cases where TEDS relies on this but may happen again the future. For the standalone version a static
@@ -48,6 +60,14 @@ void ImpersonateMassTagsAndFragments()
 
 void FTypedElementsDataStorageModule::StartupModule()
 {
+	if (!Private::GEnableTEDS)
+	{
+		UE_LOG(LogTypedElementDataStorage, Log, TEXT("Disabled by TEDS.Enable CVar"));
+		return;
+	}
+
+	UE_LOG(LogTypedElementDataStorage, Log, TEXT("Enabled by TEDS.Enable CVar"));
+	
 	ImpersonateMassTagsAndFragments();
 
 	FCoreDelegates::OnAllModuleLoadingPhasesComplete.AddLambda(
@@ -55,6 +75,8 @@ void FTypedElementsDataStorageModule::StartupModule()
 		{
 			if (!bInitialized)
 			{
+				UE_LOG(LogTypedElementDataStorage, Log, TEXT("Initializing"));
+				
 				Database = NewObject<UTypedElementDatabase>();
 				Database->Initialize();
 
@@ -113,17 +135,22 @@ void FTypedElementsDataStorageModule::StartupModule()
 					Factory->RegisterQueries(*Database);
 					Factory->RegisterWidgetConstructors(*Database, *DatabaseUi);
 				}
+				
+				UE_LOG(LogTypedElementDataStorage, Log, TEXT("Initialized"));
 
 				bInitialized = true;
 			}
 		});
 	FCoreDelegates::OnExit.AddRaw(this, &FTypedElementsDataStorageModule::ShutdownModule);
+
 }
 
 void FTypedElementsDataStorageModule::ShutdownModule()
 {
 	if (bInitialized)
 	{
+		UE_LOG(LogTypedElementDataStorage, Log, TEXT("Deinitializing"));
+		
 		UTypedElementRegistry* Registry = UTypedElementRegistry::GetInstance();
 		if (Registry) // If the registry has already been destroyed there's no point in clearing the reference.
 		{
diff --git a/Engine/Source/Runtime/TypedElementFramework/Private/Elements/Common/TypedElementDataStorageLog.cpp b/Engine/Source/Runtime/TypedElementFramework/Private/Elements/Common/TypedElementDataStorageLog.cpp
new file mode 100644
--- /dev/null
+++ b/Engine/Source/Runtime/TypedElementFramework/Private/Elements/Common/TypedElementDataStorageLog.cpp
@@ -0,0 +1,5 @@
+﻿// Copyright Epic Games, Inc. All Rights Reserved.
+
+#include "Elements/Common/TypedElementDataStorageLog.h"
+
+DEFINE_LOG_CATEGORY(LogTypedElementDataStorage)
\ No newline at end of file
diff --git a/Engine/Source/Runtime/TypedElementFramework/Public/Elements/Common/TypedElementDataStorageLog.h b/Engine/Source/Runtime/TypedElementFramework/Public/Elements/Common/TypedElementDataStorageLog.h
new file mode 100644
--- /dev/null
+++ b/Engine/Source/Runtime/TypedElementFramework/Public/Elements/Common/TypedElementDataStorageLog.h
@@ -0,0 +1,7 @@
+﻿// Copyright Epic Games, Inc. All Rights Reserved.
+
+#pragma once
+
+#include <Logging/LogMacros.h>
+
+TYPEDELEMENTFRAMEWORK_API DECLARE_LOG_CATEGORY_EXTERN(LogTypedElementDataStorage, Log, All);
\ No newline at end of file
