Editor: Fix TEDS breaking undo transactions for actors when they are deleted.

diff --git a/Engine/Plugins/Experimental/TypedElementsDataStorage/Source/TypedElementsDataStorage/Private/TypedElementDatabaseCompatibility.cpp b/Engine/Plugins/Experimental/TypedElementsDataStorage/Source/TypedElementsDataStorage/Private/TypedElementDatabaseCompatibility.cpp
--- a/Engine/Plugins/Experimental/TypedElementsDataStorage/Source/TypedElementsDataStorage/Private/TypedElementDatabaseCompatibility.cpp
+++ b/Engine/Plugins/Experimental/TypedElementsDataStorage/Source/TypedElementsDataStorage/Private/TypedElementDatabaseCompatibility.cpp
@@ -348,7 +348,7 @@ TypedElementRowHandle UTypedElementDatabaseCompatibility::AddCompatibleObjectExp
 		{
 			if (GUndo)
 			{
-				GUndo->StoreUndo(Object, MakeUnique<FRegistrationCommandChange>(Table, this));
+				GUndo->StoreUndo(this, MakeUnique<FRegistrationCommandChange>(Table, Object));
 			}
 		}
 
@@ -389,7 +389,7 @@ void UTypedElementDatabaseCompatibility::RemoveCompatibleObjectExplicitTransacti
 						// reducing the memory requirements though at the cost of more table mutations during reconstruction.
 						// This does however put more emphasis on storing any auxiliary data that can't be reconstructed from
 						// the object, e.g. the selection column, through the memento system.
-						GUndo->StoreUndo(Object, MakeUnique<FDeregistrationCommandChange>(Table->SourceTable, this));
+						GUndo->StoreUndo(this, MakeUnique<FDeregistrationCommandChange>(Table->SourceTable, Object));
 					}
 				}
 			}
@@ -726,20 +726,32 @@ void UTypedElementDatabaseCompatibility::OnPreObjectRemoved(const void* Object,
 //
 
 UTypedElementDatabaseCompatibility::FRegistrationCommandChange::FRegistrationCommandChange(
-	TypedElementDataStorage::TableHandle InTable, UTypedElementDatabaseCompatibility* InCompatibilityLayer)
-	: CompatibilityLayer(InCompatibilityLayer)
-	, Table(InTable)
+	TypedElementDataStorage::TableHandle InTable, UObject* InTargetObject)
+	: Table(InTable)
+	, TargetObject(InTargetObject)
 {
 }
 
 void UTypedElementDatabaseCompatibility::FRegistrationCommandChange::Apply(UObject* Object)
 {
-	CompatibilityLayer->AddCompatibleObjectExplicitTransactionable<false>(Object, Table);
+	if (UTypedElementDatabaseCompatibility* CompatibilityLayer = Cast<UTypedElementDatabaseCompatibility>(Object))
+	{
+		if (UObject* TargetRetrieved = TargetObject.Get(/*bEvenIfPendingKill=*/ true))
+		{
+			CompatibilityLayer->AddCompatibleObjectExplicitTransactionable<false>(TargetRetrieved, Table);
+		}
+	}
 }
 
 void UTypedElementDatabaseCompatibility::FRegistrationCommandChange::Revert(UObject* Object)
 {
-	CompatibilityLayer->RemoveCompatibleObjectExplicitTransactionable<false>(Object);
+	if (UTypedElementDatabaseCompatibility* CompatibilityLayer = Cast<UTypedElementDatabaseCompatibility>(Object))
+	{
+		if (UObject* TargetRetrieved = TargetObject.Get(/*bEvenIfPendingKill=*/ true))
+		{
+			CompatibilityLayer->RemoveCompatibleObjectExplicitTransactionable<false>(TargetRetrieved);
+		}
+	}
 }
 
 FString UTypedElementDatabaseCompatibility::FRegistrationCommandChange::ToString() const
@@ -753,20 +765,32 @@ FString UTypedElementDatabaseCompatibility::FRegistrationCommandChange::ToString
 //
 
 UTypedElementDatabaseCompatibility::FDeregistrationCommandChange::FDeregistrationCommandChange(
-	TypedElementDataStorage::TableHandle InTable, UTypedElementDatabaseCompatibility* InCompatibilityLayer)
-	: CompatibilityLayer(InCompatibilityLayer)
-	, Table(InTable)
+	TypedElementDataStorage::TableHandle InTable, UObject* InTargetObject)
+	: Table(InTable)
+	, TargetObject(InTargetObject)
 {
 }
 
 void UTypedElementDatabaseCompatibility::FDeregistrationCommandChange::Apply(UObject* Object)
 {
-	CompatibilityLayer->RemoveCompatibleObjectExplicitTransactionable<false>(Object);
+	if (UTypedElementDatabaseCompatibility* CompatibilityLayer = Cast<UTypedElementDatabaseCompatibility>(Object))
+	{
+		if (UObject* TargetRetrieved = TargetObject.Get(/*bEvenIfPendingKill=*/ true))
+		{
+			CompatibilityLayer->RemoveCompatibleObjectExplicitTransactionable<false>(TargetRetrieved);
+		}
+	}
 }
 
 void UTypedElementDatabaseCompatibility::FDeregistrationCommandChange::Revert(UObject* Object)
 {
-	CompatibilityLayer->AddCompatibleObjectExplicitTransactionable<false>(Object, Table);
+	if (UTypedElementDatabaseCompatibility* CompatibilityLayer = Cast<UTypedElementDatabaseCompatibility>(Object))
+	{
+		if (UObject* TargetRetrieved = TargetObject.Get(/*bEvenIfPendingKill=*/ true))
+		{
+			CompatibilityLayer->AddCompatibleObjectExplicitTransactionable<false>(TargetRetrieved, Table);
+		}
+	}
 }
 
 FString UTypedElementDatabaseCompatibility::FDeregistrationCommandChange::ToString() const
diff --git a/Engine/Plugins/Experimental/TypedElementsDataStorage/Source/TypedElementsDataStorage/Private/TypedElementDatabaseCompatibility.h b/Engine/Plugins/Experimental/TypedElementsDataStorage/Source/TypedElementsDataStorage/Private/TypedElementDatabaseCompatibility.h
--- a/Engine/Plugins/Experimental/TypedElementsDataStorage/Source/TypedElementsDataStorage/Private/TypedElementDatabaseCompatibility.h
+++ b/Engine/Plugins/Experimental/TypedElementsDataStorage/Source/TypedElementsDataStorage/Private/TypedElementDatabaseCompatibility.h
@@ -62,34 +62,37 @@ public:
 	TypedElementRowHandle FindRowWithCompatibleObjectExplicit(const void* Object) const override;
 
 private:
+
+	// The below changes expect UTypedElementDatabaseCompatibility to be the object passed in to StoreUndo.
+	// Note that we cannot pass in TargetObject to StoreUndo because doing so seems to stomp regular Modify()
+	// changes for that object.
 	class FRegistrationCommandChange final : public FCommandChange
 	{
 	public:
 		FRegistrationCommandChange(TypedElementDataStorage::TableHandle InTable,
-			UTypedElementDatabaseCompatibility* InCompatibilityLayer);
+			UObject* InTargetObject);
 
 		void Apply(UObject* Object) override;
 		void Revert(UObject* Object) override;
 		FString ToString() const override;
 
 	private:
-		UTypedElementDatabaseCompatibility* CompatibilityLayer{ nullptr };
 		TypedElementDataStorage::TableHandle Table{ TypedElementDataStorage::InvalidTableHandle };
+		TWeakObjectPtr<UObject> TargetObject;
 	};
-
 	class FDeregistrationCommandChange final : public FCommandChange
 	{
 	public:
 		FDeregistrationCommandChange(TypedElementDataStorage::TableHandle InTable,
-			UTypedElementDatabaseCompatibility* InCompatibilityLayer);
+			UObject* InTargetObject);
 
 		void Apply(UObject* Object) override;
 		void Revert(UObject* Object) override;
 		FString ToString() const override;
 
 	private:
-		UTypedElementDatabaseCompatibility* CompatibilityLayer{ nullptr };
 		TypedElementDataStorage::TableHandle Table{ TypedElementDataStorage::InvalidTableHandle };
+		TWeakObjectPtr<UObject> TargetObject;
 	};
 
 	void Prepare();
