Bug fix in TEDS Memento System resulting in incorrect data on reinstancing
* Default translators were index into source columns using Context.GetColumn(const UScriptStruct*) within a processor that works on single rows.  This function unfortunately does not return the column of the contextual row, it returns the start of the column in the chunk.
* This lead to issues where the mementoized data for selection would be deserialized into the wrong row which could cause objects to become part of selection sets that they were never indended to be a part of.  Example, an object in the level editor could could become part of an asset editor's selection set.
* The fix is to use the array context for the memento translator processors.

diff --git a/Engine/Plugins/Experimental/TypedElementsDataStorage/Source/TypedElementsDataStorage/Private/Memento/TypedElementMementoSystem.cpp b/Engine/Plugins/Experimental/TypedElementsDataStorage/Source/TypedElementsDataStorage/Private/Memento/TypedElementMementoSystem.cpp
--- a/Engine/Plugins/Experimental/TypedElementsDataStorage/Source/TypedElementsDataStorage/Private/Memento/TypedElementMementoSystem.cpp
+++ b/Engine/Plugins/Experimental/TypedElementsDataStorage/Source/TypedElementsDataStorage/Private/Memento/TypedElementMementoSystem.cpp
@@ -2,9 +2,11 @@
 
 #include "TypedElementMementoSystem.h"
 
+#include "Elements/Common/TypedElementDataStorageLog.h"
 #include "Elements/Framework/TypedElementQueryBuilder.h"
 #include "Memento/TypedElementMementoInterface.h"
 #include "Memento/TypedElementMementoTranslators.h"
+#include "StructArrayView.h"
 #include "TypedElementMementoRowTypes.h"
 #include "TypedElementDatabase.h"
 
@@ -71,14 +73,31 @@ void UTypedElementMementoSystem::RegisterQueries(UTypedElementDatabase& DataStor
 			Select(
 				TranslationProcessorName,
 				FObserver::OnRemove<FTypedElementMementoOnDelete>(),
-				[MementoTranslator, Memento, MementoizedColumn](TypedElementDataStorage::IQueryContext& Context, TypedElementRowHandle Row, const FTypedElementMementoOnDelete& MementoRow)
+				[MementoTranslator](TypedElementDataStorage::IQueryContext& Context, const TypedElementRowHandle* Rows, const FTypedElementMementoOnDelete* MementoTargets)
 				{
-					void* StagedMementoColumn = Context.AddColumnUninitialized(MementoRow.Memento, Memento);
-					// StagedMementoColumn will be uninitialized memory, ensure constructor called
-					Memento->InitializeStruct(StagedMementoColumn);
+					const int32 RowCount = static_cast<int32>(Context.GetRowCount());					
+					const UScriptStruct* MementoType = MementoTranslator->GetMementoType();
+					const UScriptStruct* MementoizedColumn = MementoTranslator->GetColumnType();
+
+					TArrayView<const TypedElementRowHandle> RowsView = MakeArrayView(Rows, RowCount);
+					TArrayView<const FTypedElementMementoOnDelete> MementoTargetsView = MakeArrayView(MementoTargets, RowCount);
+					FConstStructArrayView SourceColumnArrayView(*MementoizedColumn, Context.GetColumn(MementoizedColumn), RowCount);
+
+					for (int32 Index = 0; Index < RowCount; ++Index)
+					{
+						TypedElementRowHandle SourceRow = RowsView[Index];
+						TypedElementRowHandle TargetMementoRow = MementoTargetsView[Index].Memento;
+						
+						void* TargetMementoStagedColumn = Context.AddColumnUninitialized(TargetMementoRow, MementoType);
+						// StagedMementoColumn will be uninitialized memory, ensure constructor called
+						MementoType->InitializeStruct(TargetMementoStagedColumn);
 		
-					const void* SourceColumn = Context.GetColumn(MementoizedColumn);
-					MementoTranslator->TranslateColumnToMemento(SourceColumn, StagedMementoColumn);
+						const void* SourceColumn = SourceColumnArrayView[Index].GetMemory();
+
+						UE_LOG(LogTypedElementDataStorage, VeryVerbose, TEXT("Column->Memento: %llu -> %llu   0x%p -> 0x%p"), SourceRow, TargetMementoRow, SourceColumn, TargetMementoStagedColumn);
+
+						MementoTranslator->TranslateColumnToMemento(SourceColumn, TargetMementoStagedColumn);
+					}
 				})
 				.ReadOnly(MementoizedColumn)
 				.Compile());
@@ -115,16 +134,33 @@ void UTypedElementMementoSystem::RegisterQueries(UTypedElementDatabase& DataStor
 			Select(
 				TranslationProcessorName,
 				FProcessor(DSI::EQueryTickPhase::PostPhysics, DataStorage.GetQueryTickGroupName(DSI::EQueryTickGroups::Default)),
-				[MementoTranslator](TEDS::IQueryContext& Context, TypedElementRowHandle Row, const FTypedElementMementoReinstanceTarget& ReinstanceTarget)
+				[MementoTranslator](TEDS::IQueryContext& Context, const TypedElementRowHandle* Rows, const FTypedElementMementoReinstanceTarget* ReinstanceTargets)
 				{
-					// Add the column, note the column may already exist.
-					void* StagedColumn = Context.AddColumnUninitialized(ReinstanceTarget.Target, MementoTranslator->GetColumnType());
-
-					// StagedColumn will be uninitialized memory, ensure constructor called
-					MementoTranslator->GetColumnType()->InitializeStruct(StagedColumn);
-
-					const void* Memento = Context.GetColumn(MementoTranslator->GetMementoType());
-					MementoTranslator->TranslateMementoToColumn(Memento, StagedColumn);
+					const int32 RowCount = static_cast<int32>(Context.GetRowCount());					
+					const UScriptStruct* MementoType = MementoTranslator->GetMementoType();
+					const UScriptStruct* ColumnType = MementoTranslator->GetColumnType();
+
+					TArrayView<const TypedElementRowHandle> RowsView = MakeArrayView(Rows, RowCount);
+					TArrayView<const FTypedElementMementoReinstanceTarget> ReinstanceTargetsView = MakeArrayView(ReinstanceTargets, RowCount);
+					FConstStructArrayView MementoArrayView(*MementoType, Context.GetColumn(MementoType), RowCount);
+
+					for (int32 Index = 0; Index < RowCount; ++Index)
+					{
+						const TypedElementRowHandle Row = RowsView[Index];
+						const TypedElementRowHandle Target = ReinstanceTargetsView[Index].Target;
+						
+						// Add the column, note the column may already exist.
+						void* TargetStagedColumn = Context.AddColumnUninitialized(Target, ColumnType);
+
+						// StagedColumn will be uninitialized memory, ensure constructor called
+						MementoTranslator->GetColumnType()->InitializeStruct(TargetStagedColumn);
+						
+						const void* SourceMementoOfColumn = MementoArrayView[Index].GetMemory();
+
+						UE_LOG(LogTypedElementDataStorage, VeryVerbose, TEXT("Memento->Column: %llu -> %llu   0x%p -> 0x%p"), Row, Target, SourceMementoOfColumn, TargetStagedColumn);
+						
+						MementoTranslator->TranslateMementoToColumn(SourceMementoOfColumn, TargetStagedColumn);
+					}
 				})
 				.ReadOnly(MementoType)
 				.Where()
