TEDS Compat now listens to actors being destroyed.

TEDS Compat should only keep actors/uobjects around that are conceptually alive (i.e. the editor things they're alive) rather than rely on garbage collection. Using the latter exclusively will cause object that users explicitly remove to be kept around in the tools like the Outliner for longer than expected. This change adds monitoring for the explicit destruction of actors and removes that actor from TEDS if it exists rather than wait for the garbage collection to catch up.

diff --git a/Engine/Plugins/Experimental/TypedElementsDataStorage/Source/TypedElementsDataStorage/Private/TypedElementDatabaseCompatibility.cpp b/Engine/Plugins/Experimental/TypedElementsDataStorage/Source/TypedElementsDataStorage/Private/TypedElementDatabaseCompatibility.cpp
--- a/Engine/Plugins/Experimental/TypedElementsDataStorage/Source/TypedElementsDataStorage/Private/TypedElementDatabaseCompatibility.cpp
+++ b/Engine/Plugins/Experimental/TypedElementsDataStorage/Source/TypedElementsDataStorage/Private/TypedElementDatabaseCompatibility.cpp
@@ -30,10 +30,21 @@ void UTypedElementDatabaseCompatibility::Initialize(ITypedElementDataStorageInte
 
 	PostEditChangePropertyDelegateHandle = FCoreUObjectDelegates::OnObjectPropertyChanged.AddUObject(this, &UTypedElementDatabaseCompatibility::OnPostEditChangeProperty);
 	ObjectModifiedDelegateHandle = FCoreUObjectDelegates::OnObjectModified.AddUObject(this, &UTypedElementDatabaseCompatibility::OnObjectModified);
+	
+	PostWorldInitializationDelegateHandle = FWorldDelegates::OnPostWorldInitialization.AddUObject(this, &UTypedElementDatabaseCompatibility::OnPostWorldInitialization);
+	PreWorldFinishDestroyDelegateHandle = FWorldDelegates::OnPreWorldFinishDestroy.AddUObject(this, &UTypedElementDatabaseCompatibility::OnPreWorldFinishDestroy);
 }
 
 void UTypedElementDatabaseCompatibility::Deinitialize()
 {
+	for (TPair<UWorld*, FDelegateHandle>& It : ActorDestroyedDelegateHandles)
+	{
+		It.Key->RemoveOnActorDestroyededHandler(It.Value);
+	}
+
+	FWorldDelegates::OnPreWorldFinishDestroy.Remove(PreWorldFinishDestroyDelegateHandle);
+	FWorldDelegates::OnPostWorldInitialization.Remove(PostWorldInitializationDelegateHandle);
+
 	FCoreUObjectDelegates::OnObjectModified.Remove(ObjectModifiedDelegateHandle);
 	FCoreUObjectDelegates::OnObjectPropertyChanged.Remove(PostEditChangePropertyDelegateHandle);
 	
@@ -719,6 +730,27 @@ void UTypedElementDatabaseCompatibility::OnPreObjectRemoved(const void* Object,
 	}
 }
 
+void UTypedElementDatabaseCompatibility::OnPostWorldInitialization(UWorld* World, const UWorld::InitializationValues InitializationValues)
+{
+	FDelegateHandle Handle = World->AddOnActorDestroyedHandler(
+		FOnActorDestroyed::FDelegate::CreateUObject(this, &UTypedElementDatabaseCompatibility::OnActorDestroyed));
+	ActorDestroyedDelegateHandles.Add(World, Handle);
+}
+
+void UTypedElementDatabaseCompatibility::OnPreWorldFinishDestroy(UWorld* World)
+{
+	FDelegateHandle Handle;
+	if (ActorDestroyedDelegateHandles.RemoveAndCopyValue(World, Handle))
+	{
+		World->RemoveOnActorDestroyededHandler(Handle);
+	}
+}
+
+void UTypedElementDatabaseCompatibility::OnActorDestroyed(AActor* Actor)
+{
+	RemoveCompatibleObjectExplicit(Actor);
+}
+
 
 
 //
diff --git a/Engine/Plugins/Experimental/TypedElementsDataStorage/Source/TypedElementsDataStorage/Private/TypedElementDatabaseCompatibility.h b/Engine/Plugins/Experimental/TypedElementsDataStorage/Source/TypedElementsDataStorage/Private/TypedElementDatabaseCompatibility.h
--- a/Engine/Plugins/Experimental/TypedElementsDataStorage/Source/TypedElementsDataStorage/Private/TypedElementDatabaseCompatibility.h
+++ b/Engine/Plugins/Experimental/TypedElementsDataStorage/Source/TypedElementsDataStorage/Private/TypedElementDatabaseCompatibility.h
@@ -7,6 +7,7 @@
 #include "Compatibility/TypedElementObjectReinstancingManager.h"
 #include "Elements/Interfaces/TypedElementDataStorageInterface.h"
 #include "Elements/Interfaces/TypedElementDataStorageCompatibilityInterface.h"
+#include "Engine/World.h"
 #include "Misc/Change.h"
 #include "UObject/ObjectMacros.h"
 #include "UObject/WeakObjectPtr.h"
@@ -119,6 +120,11 @@ private:
 	void OnObjectAdded(const void* Object, FTypedElementDatabaseCompatibilityObjectTypeInfo TypeInfo, TypedElementRowHandle Row) const;
 	void OnPreObjectRemoved(const void* Object, FTypedElementDatabaseCompatibilityObjectTypeInfo TypeInfo, TypedElementRowHandle Row) const;
 	
+	void OnPostWorldInitialization(UWorld* World, const UWorld::InitializationValues InitializationValues);
+	void OnPreWorldFinishDestroy(UWorld* World);
+
+	void OnActorDestroyed(AActor* Actor);
+	
 	template<typename AddressType>
 	struct PendingRegistration
 	{
@@ -169,8 +175,11 @@ private:
 	 */
 	TSet<TObjectKey<const UObject>> ObjectsNeedingFullSync;
 
+	TMap<UWorld*, FDelegateHandle> ActorDestroyedDelegateHandles;
 	FDelegateHandle PostEditChangePropertyDelegateHandle;
 	FDelegateHandle ObjectModifiedDelegateHandle;
+	FDelegateHandle PostWorldInitializationDelegateHandle;
+	FDelegateHandle PreWorldFinishDestroyDelegateHandle;
 };
 
 enum class ETypedElementDatabaseCompatibilityObjectType : uint8
