Rewrote UTypedElementDatabaseUi::CreateSingleWidgetConstructor to use a function object instead of lambdas.

Functionally there's nothing wrong with UTypedElementDatabaseUi::CreateSingleWidgetConstructor, but it was triggering a false positive from static code analysis. This update to a function object achieves the same but shouldn't trigger the same warning.

diff --git a/Engine/Plugins/Experimental/TypedElementsDataStorage/Source/TypedElementsDataStorage/Private/TypedElementDatabaseUI.cpp b/Engine/Plugins/Experimental/TypedElementsDataStorage/Source/TypedElementsDataStorage/Private/TypedElementDatabaseUI.cpp
--- a/Engine/Plugins/Experimental/TypedElementsDataStorage/Source/TypedElementsDataStorage/Private/TypedElementDatabaseUI.cpp
+++ b/Engine/Plugins/Experimental/TypedElementsDataStorage/Source/TypedElementsDataStorage/Private/TypedElementDatabaseUI.cpp
@@ -362,36 +362,51 @@ bool UTypedElementDatabaseUi::CreateSingleWidgetConstructor(
 	TArray<TWeakObjectPtr<const UScriptStruct>> MatchedColumnTypes, 
 	const WidgetConstructorCallback& Callback)
 {
-	return std::visit(Internal::TOverloaded
+	struct Visitor
+	{
+		Visitor(
+			TArray<TWeakObjectPtr<const UScriptStruct>>&& InMatchedColumnTypes,
+			const TypedElementDataStorage::FMetaDataView& InArguments,
+			const WidgetConstructorCallback& InCallback) 
+			: MatchedColumnTypes(MoveTemp(InMatchedColumnTypes))
+			, Arguments(InArguments)
+			, Callback(InCallback)
+		{}
+
+		TArray<TWeakObjectPtr<const UScriptStruct>> MatchedColumnTypes;
+		const TypedElementDataStorage::FMetaDataView& Arguments;
+		const WidgetConstructorCallback& Callback;
+
+		bool operator()(const UScriptStruct* Target)
 		{
-			[this, &Arguments, &MatchedColumnTypes, &Callback](const UScriptStruct* Target)
+			TUniquePtr<FTypedElementWidgetConstructor> Result(reinterpret_cast<FTypedElementWidgetConstructor*>(
+				FMemory::Malloc(Target->GetStructureSize(), Target->GetMinAlignment())));
+			if (Result)
 			{
-				TUniquePtr<FTypedElementWidgetConstructor> Result(reinterpret_cast<FTypedElementWidgetConstructor*>(
-					FMemory::Malloc(Target->GetStructureSize(), Target->GetMinAlignment())));
-				if (Result)
-				{
-					Target->InitializeStruct(Result.Get());
-					Result->Initialize(Arguments, MoveTemp(MatchedColumnTypes));
-					return Callback(MoveTemp(Result), Result->GetMatchedColumns());
-				}
-				return true;
-			},
-			[this, &Arguments, &MatchedColumnTypes, &Callback](const TUniquePtr<FTypedElementWidgetConstructor>& Target)
+				Target->InitializeStruct(Result.Get());
+				Result->Initialize(Arguments, MoveTemp(MatchedColumnTypes));
+				return Callback(MoveTemp(Result), Result->GetMatchedColumns());
+			}
+			return true;
+		}
+
+		bool operator()(const TUniquePtr<FTypedElementWidgetConstructor>& Target)
+		{
+			const UScriptStruct* TargetType = Target->GetTypeInfo();
+			checkf(TargetType, TEXT("Expected valid type information from a widget constructor."));
+			TUniquePtr<FTypedElementWidgetConstructor> Result(reinterpret_cast<FTypedElementWidgetConstructor*>(
+				FMemory::Malloc(TargetType->GetStructureSize(), TargetType->GetMinAlignment())));
+			if (Result)
 			{
-				const UScriptStruct* TargetType = Target->GetTypeInfo();
-				checkf(TargetType, TEXT("Expected valid type information from a widget constructor."));
-				TUniquePtr<FTypedElementWidgetConstructor> Result(reinterpret_cast<FTypedElementWidgetConstructor*>(
-					FMemory::Malloc(TargetType->GetStructureSize(), TargetType->GetMinAlignment())));
-				if (Result)
-				{
-					TargetType->InitializeStruct(Result.Get());
-					TargetType->CopyScriptStruct(Result.Get(), Target.Get());
-					Result->Initialize(Arguments, MoveTemp(MatchedColumnTypes));
-					return Callback(MoveTemp(Result), Result->GetMatchedColumns());
-				}
-				return true;
+				TargetType->InitializeStruct(Result.Get());
+				TargetType->CopyScriptStruct(Result.Get(), Target.Get());
+				Result->Initialize(Arguments, MoveTemp(MatchedColumnTypes));
+				return Callback(MoveTemp(Result), Result->GetMatchedColumns());
 			}
-		}, Constructor);
+			return true;
+		}
+	};
+	return std::visit(Visitor(MoveTemp(MatchedColumnTypes), Arguments, Callback), Constructor);
 }
 
 void UTypedElementDatabaseUi::CreateWidgetInstance(
