Fixed actor reinstancing not registring correctly with TEDS.

When actors were re-instanced they would be detected by TEDS but not added again. The original object would however be deleted leading to objects being lost. In the TEDS Outliner this presented as "ghost actors" being added as the TEDS Outliner tried to compensate for the missing information. To fix this issue the following changes were made:
- If TEDS' re-instancing manager doesn't find an object in TEDS it will now add it. Further changes are to make sure the right table is used in this case.
- Type information can now be used to associate a (derived) type with a specific table in TEDS Compat. If an object is registered this will now be exclusively used to assign the initial table. This now also allows more specialization for different types of objects.
- It's now no longer possible to register an object with TEDS Compat with a custom table. This is to avoid problems with the wrong table being used when an object is reconstructed e.g. through re-instancing. RegisterTypeTableAssociation can now be used to achieve the same effect.
- Undo/redo for deleting/adding objects now uses the type->table association rather than store a table. The same table would be used in both cases as they both use table used to create the object with, not any later mutations.
- All handling of actors (except for a specific case for the reference handling) is now unified with UObjects. This was already largely done but behind an #ifdef for testing purposes. This mostly just removes the unused code.
- The factory function "RegisterTables" now has an overloaded function that also takes a ITypedElementDataStorageCompatibilityInterface to make it easier to associate a type with a table.

diff --git a/Engine/Plugins/Experimental/TypedElementsDataStorage/Source/TypedElementsDataStorage/Private/Compatibility/TypedElementObjectReinstancingManager.cpp b/Engine/Plugins/Experimental/TypedElementsDataStorage/Source/TypedElementsDataStorage/Private/Compatibility/TypedElementObjectReinstancingManager.cpp
--- a/Engine/Plugins/Experimental/TypedElementsDataStorage/Source/TypedElementsDataStorage/Private/Compatibility/TypedElementObjectReinstancingManager.cpp
+++ b/Engine/Plugins/Experimental/TypedElementsDataStorage/Source/TypedElementsDataStorage/Private/Compatibility/TypedElementObjectReinstancingManager.cpp
@@ -1,4 +1,4 @@
-﻿// Copyright Epic Games, Inc. All Rights Reserved.
+// Copyright Epic Games, Inc. All Rights Reserved.
 
 #include "TypedElementObjectReinstancingManager.h"
 
@@ -92,21 +92,23 @@ void UTypedElementObjectReinstancingManager::HandleOnObjectsReinstanced(
 {
 	for (FCoreUObjectDelegates::FReplacementObjectMap::TConstIterator Iter = ObjectReplacementMap.CreateConstIterator(); Iter; ++Iter)
 	{
-		const UObject* NewInstanceObject = Iter->Value;
+		UObject* NewInstanceObject = Iter->Value;
 
 		TypedElementRowHandle NewObjectRow = DataStorageCompatibility->FindRowWithCompatibleObjectExplicit(NewInstanceObject);
-		if (NewObjectRow != TypedElementInvalidRowHandle)
+		if (!Database->IsRowAvailable(NewObjectRow))
 		{
-			const void* PreDeleteObject = Iter->Key;
-			const TypedElementRowHandle* MementoRowPtr = OldObjectToMementoMap.Find(PreDeleteObject);
-			if (MementoRowPtr != nullptr)
+			NewObjectRow = DataStorageCompatibility->AddCompatibleObjectExplicit(NewInstanceObject);
+		}
+		
+		const void* PreDeleteObject = Iter->Key;
+		const TypedElementRowHandle* MementoRowPtr = OldObjectToMementoMap.Find(PreDeleteObject);
+		if (MementoRowPtr != nullptr)
+		{
+			// Kick off re-instantiation of NewObjectRow from the Memento
+			TypedElementRowHandle Memento = *MementoRowPtr;
+			if (ensureMsgf(Database->HasColumns(Memento, TConstArrayView<const UScriptStruct*>({FTypedElementMementoTag::StaticStruct()})), TEXT("Cannot reinstantiate from a non memento row")))
 			{
-				// Kick off reinstantiation of NewObjectRow from the Memento
-				TypedElementRowHandle Memento = *MementoRowPtr;
-				if (ensureMsgf(Database->HasColumns(Memento, TConstArrayView<const UScriptStruct*>({FTypedElementMementoTag::StaticStruct()})), TEXT("Cannot reinstantiate from a non memento row")))
-				{
-					Database->AddOrGetColumn(Memento, FTypedElementMementoReinstanceTarget{ .Target = NewObjectRow });
-				}
+				Database->AddOrGetColumn(Memento, FTypedElementMementoReinstanceTarget{ .Target = NewObjectRow });
 			}
 		}
 	}
diff --git a/Engine/Plugins/Experimental/TypedElementsDataStorage/Source/TypedElementsDataStorage/Private/TypedElementDatabaseCompatibility.cpp b/Engine/Plugins/Experimental/TypedElementsDataStorage/Source/TypedElementsDataStorage/Private/TypedElementDatabaseCompatibility.cpp
--- a/Engine/Plugins/Experimental/TypedElementsDataStorage/Source/TypedElementsDataStorage/Private/TypedElementDatabaseCompatibility.cpp
+++ b/Engine/Plugins/Experimental/TypedElementsDataStorage/Source/TypedElementsDataStorage/Private/TypedElementDatabaseCompatibility.cpp
@@ -15,10 +15,6 @@
 #include "MassActorSubsystem.h"
 #include "TypedElementDataStorageProfilingMacros.h"
 
-#if TEDS_SEPARATE_ACTOR_REGISTRATION
-#	include "MassActorEditorSubsystem.h"
-#endif
-
 void UTypedElementDatabaseCompatibility::Initialize(ITypedElementDataStorageInterface* StorageInterface)
 {
 	checkf(StorageInterface, TEXT("Typed Element's Database compatibility manager is being initialized with an invalid storage target."));
@@ -61,6 +57,12 @@ void UTypedElementDatabaseCompatibility::RegisterDealiaserCallback(ObjectToRowDe
 	ObjectToRowDialiasers.Add(MoveTemp(Dealiaser));
 }
 
+void UTypedElementDatabaseCompatibility::RegisterTypeTableAssociation(
+	TObjectPtr<UStruct> TypeInfo, TypedElementDataStorage::TableHandle Table)
+{
+	TypeToTableMap.Add(TypeInfo, Table);
+}
+
 FDelegateHandle UTypedElementDatabaseCompatibility::RegisterObjectAddedCallback(ObjectAddedCallback&& OnObjectAdded)
 {
 	FDelegateHandle Handle(FDelegateHandle::GenerateNewHandle);
@@ -93,55 +95,21 @@ void UTypedElementDatabaseCompatibility::UnregisterObjectRemovedCallback(FDelega
 
 TypedElementRowHandle UTypedElementDatabaseCompatibility::AddCompatibleObjectExplicit(UObject* Object)
 {
-	return AddCompatibleObjectExplicit(Object, StandardUObjectTable);
-}
-
-TypedElementRowHandle UTypedElementDatabaseCompatibility::AddCompatibleObjectExplicit(UObject* Object, TypedElementTableHandle Table)
-{
-	bool bCanAddObject = 
+	bool bCanAddObject =
 		ensureMsgf(Storage, TEXT("Trying to add a UObject to Typed Element's Data Storage before the storage is available.")) &&
 		ShouldAddObject(Object);
-	return bCanAddObject ? AddCompatibleObjectExplicitTransactionable<true>(Object, Table) : TypedElementDataStorage::InvalidRowHandle;
-}
-
-TypedElementRowHandle UTypedElementDatabaseCompatibility::AddCompatibleObjectExplicit(AActor* Actor)
-{
-	return AddCompatibleObjectExplicit(Actor, StandardActorTable);
-}
-
-TypedElementRowHandle UTypedElementDatabaseCompatibility::AddCompatibleObjectExplicit(AActor* Actor, TypedElementTableHandle Table)
-{
-#if TEDS_SEPARATE_ACTOR_REGISTRATION
-	if (ensureMsgf(Storage, TEXT("Trying to add an actor to Typed Element's Data Storage before the storage is available.")) &&
-		ShouldAddObject(Actor))
-	{
-		// Registration is delayed for two reasons:
-		//	1. Allows entity creation in a single batch rather than multiple individual additions.
-		//	2. Provides an opportunity to filter out the actors that are created within MASS itself as those will already be registered.
-		
-		TypedElementRowHandle ReservedRow = Storage->ReserveRow();
-		PendingRegistration<TWeakObjectPtr<AActor>>& Pending = ActorsPendingRegistration.FindOrAdd(Table);
-		Pending.Add(ReservedRow, Actor);
-		return ReservedRow;
-	}
-	return TypedElementInvalidRowHandle;
-#else
-	return AddCompatibleObjectExplicit(static_cast<UObject*>(Actor), Table);
-#endif
+	return bCanAddObject ? AddCompatibleObjectExplicitTransactionable<true>(Object) : TypedElementDataStorage::InvalidRowHandle;
 }
 
 TypedElementRowHandle UTypedElementDatabaseCompatibility::AddCompatibleObjectExplicit(void* Object, TWeakObjectPtr<const UScriptStruct> TypeInfo)
-{
-	return AddCompatibleObjectExplicit(Object, MoveTemp(TypeInfo), StandardExternalObjectTable);
-}
-
-TypedElementRowHandle UTypedElementDatabaseCompatibility::AddCompatibleObjectExplicit(
-	void* Object, TWeakObjectPtr<const UScriptStruct> TypeInfo, TypedElementTableHandle Table)
 {
 	using namespace TypedElementDataStorage;
-
+	
 	if (ensureMsgf(Storage, TEXT("Trying to add an object to Typed Element's Data Storage before the storage is available.")))
 	{
+		TableHandle Table = FindBestMatchingTable(TypeInfo.Get());
+		Table = (Table != InvalidTableHandle) ? Table : StandardExternalObjectTable;
+
 		TypedElementRowHandle ReservedRow = Storage->ReserveRow();
 		Storage->IndexRow(GenerateIndexHash(Object), ReservedRow);
 		PendingRegistration<ExternalObjectRegistration>& Pending = ExternalObjectsPendingRegistration.FindOrAdd(Table);
@@ -156,19 +124,7 @@ TypedElementRowHandle UTypedElementDatabaseCompatibility::AddCompatibleObjectExp
 
 void UTypedElementDatabaseCompatibility::RemoveCompatibleObjectExplicit(UObject* Object)
 {
-	using namespace TypedElementDataStorage;
-
-#if TEDS_SEPARATE_ACTOR_REGISTRATION
-	if (Object->IsA<AActor>())
-	{
-		RemoveCompatibleObjectExplicit(static_cast<AActor*>(Object));
-	}
-	else
-#else
-	{
-		RemoveCompatibleObjectExplicitTransactionable<true>(Object);
-	}
-#endif
+	RemoveCompatibleObjectExplicitTransactionable<true>(Object);
 }
 
 void UTypedElementDatabaseCompatibility::RemoveCompatibleObjectExplicit(void* Object)
@@ -189,42 +145,6 @@ void UTypedElementDatabaseCompatibility::RemoveCompatibleObjectExplicit(void* Ob
 	}
 }
 
-void UTypedElementDatabaseCompatibility::RemoveCompatibleObjectExplicit(AActor* Actor)
-{
-#if TEDS_SEPARATE_ACTOR_REGISTRATION
-	checkf(Storage, TEXT("Removing compatible objects is not supported before Typed Element's Database compatibility manager has been initialized."));
-
-	// If there is no actor subsystem it means that the world has been destroyed, including the MASS instance,
-	// so there's no references to clean up.
-	if (Storage && ActorSubsystem && Storage->IsAvailable())
-	{
-		FMassEntityHandle Entity = ActorSubsystem->GetEntityHandleFromActor(Actor);
-		// If there's no entity it may:
-		//	- have been deleted earlier, e.g. through an explicit delete.
-		//	- be an actor that never had a world assigned and was therefore never registered.
-		//	- have registered with a MASS instance in another world, e.g. one created for PIE.
-		if (Entity.IsValid())
-		{
-			auto ActorStore = Storage->GetColumn<FMassActorFragment>(Entity.AsNumber());
-			if (ActorStore && !ActorStore->IsOwnedByMass()) // Only remove actors that were externally created.
-			{
-				TypedElementRowHandle Row = Entity.AsNumber();
-
-				const FTypedElementClassTypeInfoColumn* TypeInfoColumn = Storage->GetColumn<FTypedElementClassTypeInfoColumn>(Row);
-				if (Storage->HasRowBeenAssigned(Row) && ensureMsgf(TypeInfoColumn, TEXT("Missing type information for removed actor at ptr 0x%p [%s]"), Actor, *Actor->GetName()))
-				{
-					OnPreObjectRemoved(Actor, TypeInfoColumn->TypeInfo.Get(), Row);
-				}
-				
-				ActorSubsystem->RemoveHandleForActor(Actor);
-				Storage->RemoveRow(Entity.AsNumber());
-			}
-		}
-	}
-#else
-	RemoveCompatibleObjectExplicit(static_cast<UObject*>(Actor));
-#endif
-}
 
 TypedElementRowHandle UTypedElementDatabaseCompatibility::FindRowWithCompatibleObjectExplicit(const UObject* Object) const
 {
@@ -232,42 +152,12 @@ TypedElementRowHandle UTypedElementDatabaseCompatibility::FindRowWithCompatibleO
 
 	if (Object && Storage && Storage->IsAvailable())
 	{
-#if TEDS_SEPARATE_ACTOR_REGISTRATION
-		const AActor* Actor = Cast<AActor>(Object);
-		if (Actor && ActorSubsystem)
-		{
-			FMassEntityHandle Entity = ActorSubsystem->GetEntityHandleFromActor(Actor);
-			return Entity.IsValid() ? Entity.AsNumber() : DealiasObject(Object);
-		}
-		else
-#endif
-		{
-			RowHandle Row = Storage->FindIndexedRow(GenerateIndexHash(Object));
-			return Storage->IsRowAvailable(Row) ? Row : DealiasObject(Object);
-		}
-
-
+		RowHandle Row = Storage->FindIndexedRow(GenerateIndexHash(Object));
+		return Storage->IsRowAvailable(Row) ? Row : DealiasObject(Object);
 	}
 	return TypedElementInvalidRowHandle;
 }
 
-TypedElementRowHandle UTypedElementDatabaseCompatibility::FindRowWithCompatibleObjectExplicit(const AActor* Actor) const
-{
-#if TEDS_SEPARATE_ACTOR_REGISTRATION
-	if (Storage && ActorSubsystem && Storage->IsAvailable())
-	{
-		FMassEntityHandle Entity = ActorSubsystem->GetEntityHandleFromActor(Actor);
-		return Entity.IsValid() ? Entity.AsNumber() : DealiasObject(Actor);
-	}
-	else
-	{
-		return TypedElementInvalidRowHandle;
-	}
-#else
-	return FindRowWithCompatibleObjectExplicit(static_cast<const UObject*>(Actor));
-#endif
-}
-
 TypedElementRowHandle UTypedElementDatabaseCompatibility::FindRowWithCompatibleObjectExplicit(const void* Object) const
 {
 	using namespace TypedElementDataStorage;
@@ -275,29 +165,29 @@ TypedElementRowHandle UTypedElementDatabaseCompatibility::FindRowWithCompatibleO
 	return (Object && Storage && Storage->IsAvailable()) ? Storage->FindIndexedRow(GenerateIndexHash(Object)) : InvalidRowHandle;
 }
 
-void UTypedElementDatabaseCompatibility::Prepare()
+void UTypedElementDatabaseCompatibility::AddReferencedObjects(UObject* InThis, FReferenceCollector& Collector)
 {
-#if TEDS_SEPARATE_ACTOR_REGISTRATION
-	UMassActorEditorSubsystem* MassActorEditorSubsystem = Storage->GetExternalSystem<UMassActorEditorSubsystem>();
-	check(MassActorEditorSubsystem);
-	ActorSubsystem = MassActorEditorSubsystem->GetMutableActorManager().AsShared();
-#endif
+	Super::AddReferencedObjects(InThis, Collector);
+
+	UTypedElementDatabaseCompatibility* Compatibility = Cast<UTypedElementDatabaseCompatibility>(InThis);
+	checkf(Compatibility, TEXT("AddReferencedObjects was given a null pointer or an object wasn't a typed element data compatibility object."));
+
+	Collector.AddReferencedObjects(Compatibility->TypeToTableMap);
+}
 
+void UTypedElementDatabaseCompatibility::Prepare()
+{
 	CreateStandardArchetypes();
 }
 
 void UTypedElementDatabaseCompatibility::Reset()
 {
-#if TEDS_SEPARATE_ACTOR_REGISTRATION
-	ActorSubsystem = nullptr;
-#endif
 }
 
 void UTypedElementDatabaseCompatibility::CreateStandardArchetypes()
 {
 	StandardActorTable = Storage->RegisterTable(TTypedElementColumnTypeList<
 			FMassActorFragment, FTypedElementUObjectColumn, FTypedElementClassTypeInfoColumn,
-			FTypedElementObjectSourceTableColumn,
 			FTypedElementLabelColumn, FTypedElementLabelHashColumn,
 			FTypedElementSyncFromWorldTag>(), 
 		FName("Editor_StandardActorTable"));
@@ -308,15 +198,16 @@ void UTypedElementDatabaseCompatibility::CreateStandardArchetypes()
 
 	StandardUObjectTable = Storage->RegisterTable(TTypedElementColumnTypeList<
 			FTypedElementUObjectColumn, FTypedElementClassTypeInfoColumn,
-			FTypedElementObjectSourceTableColumn,
 			FTypedElementSyncFromWorldTag>(), 
 		FName("Editor_StandardUObjectTable"));
 
 	StandardExternalObjectTable = Storage->RegisterTable(TTypedElementColumnTypeList<
 			FTypedElementExternalObjectColumn, FTypedElementScriptStructTypeInfoColumn,
-			FTypedElementObjectSourceTableColumn,
 			FTypedElementSyncFromWorldTag>(), 
 		FName("Editor_StandardExternalObjectTable"));
+
+	RegisterTypeTableAssociation(AActor::StaticClass(), StandardActorTable);
+	RegisterTypeTableAssociation(UObject::StaticClass(), StandardUObjectTable);
 }
 
 bool UTypedElementDatabaseCompatibility::ShouldAddObject(const UObject* Object) const
@@ -336,35 +227,46 @@ bool UTypedElementDatabaseCompatibility::ShouldAddObject(const UObject* Object)
 	return Include;
 }
 
-template<bool bEnableTransactions>
-TypedElementRowHandle UTypedElementDatabaseCompatibility::AddCompatibleObjectExplicitTransactionable(UObject* Object, TypedElementTableHandle Table)
+TypedElementDataStorage::TableHandle UTypedElementDatabaseCompatibility::FindBestMatchingTable(const UStruct* TypeInfo) const
 {
 	using namespace TypedElementDataStorage;
 
-#if TEDS_SEPARATE_ACTOR_REGISTRATION
-	if (Object->IsA<AActor>())
+	while (TypeInfo)
 	{
-		return AddCompatibleObjectExplicit(static_cast<AActor*>(Object));
+		if (const TableHandle* Table = TypeToTableMap.Find(TypeInfo))
+		{
+			return *Table;
+		}
+		TypeInfo = TypeInfo->GetSuperStruct();
 	}
-	else
-#endif
-	{
-		TypedElementRowHandle ReservedRow = Storage->ReserveRow();
-		Storage->IndexRow(GenerateIndexHash(Object), ReservedRow);
 
-		PendingRegistration<TWeakObjectPtr<UObject>>& Pending = UObjectsPendingRegistration.FindOrAdd(Table);
-		Pending.Add(ReservedRow, Object);
+	return InvalidTableHandle;
+}
+
+template<bool bEnableTransactions>
+TypedElementRowHandle UTypedElementDatabaseCompatibility::AddCompatibleObjectExplicitTransactionable(UObject* Object)
+{
+	using namespace TypedElementDataStorage;
+
+	TableHandle Table = FindBestMatchingTable(Object->GetClass());
+	checkf(Table != InvalidTableHandle, TEXT("The Typed Elements Data Storage could not find any matching tables for object of type '%s'. "
+		"This can mean that the object doesn't derive from UObject or that a table for UObject is no longer registered."), *Object->GetClass()->GetFName().ToString());
+
+	TypedElementRowHandle ReservedRow = Storage->ReserveRow();
+	Storage->IndexRow(GenerateIndexHash(Object), ReservedRow);
+
+	PendingRegistration<TWeakObjectPtr<UObject>>& Pending = UObjectsPendingRegistration.FindOrAdd(Table);
+	Pending.Add(ReservedRow, Object);
 
-		if constexpr (bEnableTransactions)
+	if constexpr (bEnableTransactions)
+	{
+		if (GUndo)
 		{
-			if (GUndo)
-			{
-				GUndo->StoreUndo(this, MakeUnique<FRegistrationCommandChange>(Table, Object));
-			}
+			GUndo->StoreUndo(this, MakeUnique<FRegistrationCommandChange>(Object));
 		}
-
-		return ReservedRow;
 	}
+
+	return ReservedRow;
 }
 
 template<bool bEnableTransactions>
@@ -388,20 +290,7 @@ void UTypedElementDatabaseCompatibility::RemoveCompatibleObjectExplicitTransacti
 			{
 				if (GUndo)
 				{
-					FTypedElementObjectSourceTableColumn* Table = Storage->GetColumn<FTypedElementObjectSourceTableColumn>(Row);
-					if (ensureMsgf(Table,
-						TEXT("An object is removed from the Typed Element's Database compatibility manager that didn't contain a source table column.")))
-					{
-						// Reverting the deletion of a row relies on the original table used to create the row rather
-						// than the mutated variant that may have been used at the time of destruction. The reasoning is
-						// that using the original table will trigger the same update that the mutated versions will call
-						// so both will result in the same mirrored data in the data storage after processing. The benefit
-						// of taking the original data is that less information has to be stored in the memento table thus
-						// reducing the memory requirements though at the cost of more table mutations during reconstruction.
-						// This does however put more emphasis on storing any auxiliary data that can't be reconstructed from
-						// the object, e.g. the selection column, through the memento system.
-						GUndo->StoreUndo(this, MakeUnique<FDeregistrationCommandChange>(Table->SourceTable, Object));
-					}
+					GUndo->StoreUndo(this, MakeUnique<FDeregistrationCommandChange>(Object));
 				}
 			}
 		}
@@ -430,9 +319,6 @@ void UTypedElementDatabaseCompatibility::Tick()
 	// Delay processing until the required systems are available by not clearing any lists or doing any work.
 	if (Storage && Storage->IsAvailable() && EditorWorld)
 	{
-#if TEDS_SEPARATE_ACTOR_REGISTRATION
-		TickPendingActorRegistration(EditorWorld);
-#endif
 		TickPendingUObjectRegistration();
 		TickPendingExternalObjectRegistration();
 		TickObjectSync();
@@ -531,54 +417,6 @@ void UTypedElementDatabaseCompatibility::PendingRegistration<AddressType>::Reset
 	ReservedRowHandles.Reset();
 }
 
-#if TEDS_SEPARATE_ACTOR_REGISTRATION
-void UTypedElementDatabaseCompatibility::TickPendingActorRegistration(UWorld* EditorWorld)
-{
-	if (!ActorsPendingRegistration.IsEmpty())
-	{
-		// Filter out the actors that are already registered or already destroyed. 
-		// The most common case for this is actors created from within MASS.
-		for (auto It = ActorsPendingRegistration.CreateIterator(); It; ++It)
-		{
-			It->Value.RemoveInvalidEntries(*Storage,
-				[this, EditorWorld](const TWeakObjectPtr<AActor>& Actor)
-				{
-					AActor* Instance = Actor.Get();
-					return Instance && Instance->GetWorld() == EditorWorld && !ActorSubsystem->GetEntityHandleFromActor(Instance).IsValid();
-				});
-		}
-
-		// Add the remaining actors to the data storage.
-		for (auto It = ActorsPendingRegistration.CreateIterator(); It; ++It)
-		{
-			It->Value.ProcessEntries(*Storage, It->Key,
-				[this](TypedElementRowHandle Row, const TWeakObjectPtr<AActor>& ActorPtr)
-				{
-					FMassActorFragment* ActorStore = Storage->AddOrGetColumn<FMassActorFragment>(Row);
-					checkf(ActorStore, TEXT("Failed to retrieve or add FMassActorFragment to newly created row."));
-
-					constexpr bool bIsOwnedByMass = false;
-
-					AActor* Actor = ActorPtr.Get();
-					check(Actor);
-					ActorStore->SetNoHandleMapUpdate(FMassEntityHandle::FromNumber(Row), Actor, bIsOwnedByMass);
-					ActorSubsystem->SetHandleForActor(Actor, FMassEntityHandle::FromNumber(Row));
-
-					Storage->AddOrGetColumn<FTypedElementUObjectColumn>(Row, FTypedElementUObjectColumn{ .Object = ActorPtr });
-					Storage->AddOrGetColumn<FTypedElementClassTypeInfoColumn>(Row, FTypedElementClassTypeInfoColumn{ .TypeInfo = Actor->GetClass() });
-					
-					// Make sure the new row is tagged for update.
-					Storage->AddColumn<FTypedElementSyncFromWorldTag>(Row);
-
-					OnObjectAdded(ActorPtr.Get(), Actor->GetClass(), Row);
-				});
-		}
-			
-		ActorsPendingRegistration.Reset();
-	}
-}
-#endif
-
 void UTypedElementDatabaseCompatibility::TickPendingUObjectRegistration()
 {
 	if (!UObjectsPendingRegistration.IsEmpty())
@@ -599,16 +437,14 @@ void UTypedElementDatabaseCompatibility::TickPendingUObjectRegistration()
 			It->Value.ProcessEntries(*Storage, It->Key,
 				[Table = It->Key, this](TypedElementRowHandle Row, const TWeakObjectPtr<UObject>& Object)
 				{
-#if !TEDS_SEPARATE_ACTOR_REGISTRATION
 					if (AActor* Actor = Cast<AActor>(Object))
 					{
 						constexpr bool bIsOwnedByMass = false;
 						Storage->AddOrGetColumn<FMassActorFragment>(Row)->SetNoHandleMapUpdate(FMassEntityHandle::FromNumber(Row), Actor, bIsOwnedByMass);
 					}
-#endif
+
 					Storage->AddOrGetColumn<FTypedElementUObjectColumn>(Row, FTypedElementUObjectColumn{ .Object = Object });
 					Storage->AddOrGetColumn<FTypedElementClassTypeInfoColumn>(Row, FTypedElementClassTypeInfoColumn{ .TypeInfo = Object->GetClass() });
-					Storage->AddOrGetColumn<FTypedElementObjectSourceTableColumn>(Row, FTypedElementObjectSourceTableColumn{ .SourceTable = Table });
 					// Make sure the new row is tagged for update.
 					Storage->AddColumn<FTypedElementSyncFromWorldTag>(Row);
 					OnObjectAdded(Object.Get(), Object->GetClass(), Row);
@@ -641,7 +477,6 @@ void UTypedElementDatabaseCompatibility::TickPendingExternalObjectRegistration()
 				{
 					Storage->AddOrGetColumn<FTypedElementExternalObjectColumn>(Row, FTypedElementExternalObjectColumn{ .Object = Object.Object });
 					Storage->AddOrGetColumn<FTypedElementScriptStructTypeInfoColumn>(Row, FTypedElementScriptStructTypeInfoColumn{ .TypeInfo = Object.TypeInfo });
-					Storage->AddOrGetColumn<FTypedElementObjectSourceTableColumn>(Row, FTypedElementObjectSourceTableColumn{ .SourceTable = Table });
 					// Make sure the new row is tagged for update.
 					Storage->AddColumn<FTypedElementSyncFromWorldTag>(Row);
 
@@ -757,10 +592,8 @@ void UTypedElementDatabaseCompatibility::OnActorDestroyed(AActor* Actor)
 // UTypedElementDatabaseCompatibility::FRegistrationCommandChange
 //
 
-UTypedElementDatabaseCompatibility::FRegistrationCommandChange::FRegistrationCommandChange(
-	TypedElementDataStorage::TableHandle InTable, UObject* InTargetObject)
-	: Table(InTable)
-	, TargetObject(InTargetObject)
+UTypedElementDatabaseCompatibility::FRegistrationCommandChange::FRegistrationCommandChange(UObject* InTargetObject)
+	: TargetObject(InTargetObject)
 {
 }
 
@@ -770,7 +603,7 @@ void UTypedElementDatabaseCompatibility::FRegistrationCommandChange::Apply(UObje
 	{
 		if (UObject* TargetRetrieved = TargetObject.Get(/*bEvenIfPendingKill=*/ true))
 		{
-			CompatibilityLayer->AddCompatibleObjectExplicitTransactionable<false>(TargetRetrieved, Table);
+			CompatibilityLayer->AddCompatibleObjectExplicitTransactionable<false>(TargetRetrieved);
 		}
 	}
 }
@@ -796,10 +629,8 @@ FString UTypedElementDatabaseCompatibility::FRegistrationCommandChange::ToString
 // UTypedElementDatabaseCompatibility::FDeregistrationCommandChange
 //
 
-UTypedElementDatabaseCompatibility::FDeregistrationCommandChange::FDeregistrationCommandChange(
-	TypedElementDataStorage::TableHandle InTable, UObject* InTargetObject)
-	: Table(InTable)
-	, TargetObject(InTargetObject)
+UTypedElementDatabaseCompatibility::FDeregistrationCommandChange::FDeregistrationCommandChange(UObject* InTargetObject)
+	: TargetObject(InTargetObject)
 {
 }
 
@@ -820,7 +651,7 @@ void UTypedElementDatabaseCompatibility::FDeregistrationCommandChange::Revert(UO
 	{
 		if (UObject* TargetRetrieved = TargetObject.Get(/*bEvenIfPendingKill=*/ true))
 		{
-			CompatibilityLayer->AddCompatibleObjectExplicitTransactionable<false>(TargetRetrieved, Table);
+			CompatibilityLayer->AddCompatibleObjectExplicitTransactionable<false>(TargetRetrieved);
 		}
 	}
 }
diff --git a/Engine/Plugins/Experimental/TypedElementsDataStorage/Source/TypedElementsDataStorage/Private/TypedElementDatabaseCompatibility.h b/Engine/Plugins/Experimental/TypedElementsDataStorage/Source/TypedElementsDataStorage/Private/TypedElementDatabaseCompatibility.h
--- a/Engine/Plugins/Experimental/TypedElementsDataStorage/Source/TypedElementsDataStorage/Private/TypedElementDatabaseCompatibility.h
+++ b/Engine/Plugins/Experimental/TypedElementsDataStorage/Source/TypedElementsDataStorage/Private/TypedElementDatabaseCompatibility.h
@@ -14,11 +14,6 @@
 
 #include "TypedElementDatabaseCompatibility.generated.h"
 
-#if !defined(TEDS_SEPARATE_ACTOR_REGISTRATION)
-#	define TEDS_SEPARATE_ACTOR_REGISTRATION 0
-#endif
-
-class AActor;
 class ITypedElementDataStorageInterface;
 struct FMassActorManager;
 
@@ -42,26 +37,23 @@ public:
 
 	void RegisterRegistrationFilter(ObjectRegistrationFilter Filter) override;
 	void RegisterDealiaserCallback(ObjectToRowDealiaser Dealiaser) override;
+	void RegisterTypeTableAssociation(TObjectPtr<UStruct> TypeInfo, TypedElementDataStorage::TableHandle Table) override;
 	FDelegateHandle RegisterObjectAddedCallback(ObjectAddedCallback&& OnObjectAdded);
 	void UnregisterObjectAddedCallback(FDelegateHandle Handle);
 	FDelegateHandle RegisterObjectRemovedCallback(ObjectRemovedCallback&& OnObjectRemoved);
 	void UnregisterObjectRemovedCallback(FDelegateHandle Handle);
 	
 	TypedElementRowHandle AddCompatibleObjectExplicit(UObject* Object) override;
-	TypedElementRowHandle AddCompatibleObjectExplicit(UObject* Object, TypedElementTableHandle Table) override;
-	TypedElementRowHandle AddCompatibleObjectExplicit(AActor* Actor) override;
-	TypedElementRowHandle AddCompatibleObjectExplicit(AActor* Actor, TypedElementTableHandle Table) override;
 	TypedElementRowHandle AddCompatibleObjectExplicit(void* Object, TWeakObjectPtr<const UScriptStruct> TypeInfo);
-	TypedElementRowHandle AddCompatibleObjectExplicit(void* Object, TWeakObjectPtr<const UScriptStruct> TypeInfo, TypedElementTableHandle Table);
 	
 	void RemoveCompatibleObjectExplicit(UObject* Object) override;
-	void RemoveCompatibleObjectExplicit(AActor* Actor) override;
 	void RemoveCompatibleObjectExplicit(void* Object) override;
 
 	TypedElementRowHandle FindRowWithCompatibleObjectExplicit(const UObject* Object) const override;
-	TypedElementRowHandle FindRowWithCompatibleObjectExplicit(const AActor* Actor) const override;
 	TypedElementRowHandle FindRowWithCompatibleObjectExplicit(const void* Object) const override;
 
+	static void AddReferencedObjects(UObject* InThis, FReferenceCollector& Collector);
+
 private:
 
 	// The below changes expect UTypedElementDatabaseCompatibility to be the object passed in to StoreUndo.
@@ -70,29 +62,25 @@ private:
 	class FRegistrationCommandChange final : public FCommandChange
 	{
 	public:
-		FRegistrationCommandChange(TypedElementDataStorage::TableHandle InTable,
-			UObject* InTargetObject);
+		explicit FRegistrationCommandChange(UObject* InTargetObject);
 
 		void Apply(UObject* Object) override;
 		void Revert(UObject* Object) override;
 		FString ToString() const override;
 
 	private:
-		TypedElementDataStorage::TableHandle Table{ TypedElementDataStorage::InvalidTableHandle };
 		TWeakObjectPtr<UObject> TargetObject;
 	};
 	class FDeregistrationCommandChange final : public FCommandChange
 	{
 	public:
-		FDeregistrationCommandChange(TypedElementDataStorage::TableHandle InTable,
-			UObject* InTargetObject);
+		explicit FDeregistrationCommandChange(UObject* InTargetObject);
 
 		void Apply(UObject* Object) override;
 		void Revert(UObject* Object) override;
 		FString ToString() const override;
 
 	private:
-		TypedElementDataStorage::TableHandle Table{ TypedElementDataStorage::InvalidTableHandle };
 		TWeakObjectPtr<UObject> TargetObject;
 	};
 
@@ -101,16 +89,14 @@ private:
 	void CreateStandardArchetypes();
 	
 	bool ShouldAddObject(const UObject* Object) const;
+	TypedElementDataStorage::TableHandle FindBestMatchingTable(const UStruct* TypeInfo) const;
 	template<bool bEnableTransactions>
-	TypedElementRowHandle AddCompatibleObjectExplicitTransactionable(UObject* Object, TypedElementTableHandle Table);
+	TypedElementRowHandle AddCompatibleObjectExplicitTransactionable(UObject* Object);
 	template<bool bEnableTransactions>
 	void RemoveCompatibleObjectExplicitTransactionable(UObject* Object);
 	TypedElementRowHandle DealiasObject(const UObject* Object) const;
 
 	void Tick();
-#if TEDS_SEPARATE_ACTOR_REGISTRATION
-	void TickPendingActorRegistration(UWorld* EditorWorld);
-#endif
 	void TickPendingUObjectRegistration();
 	void TickPendingExternalObjectRegistration();
 	void TickObjectSync();
@@ -149,14 +135,12 @@ private:
 		void* Object;
 		TWeakObjectPtr<const UScriptStruct> TypeInfo;
 	};
-#if TEDS_SEPARATE_ACTOR_REGISTRATION
-	TMap<TypedElementTableHandle, PendingRegistration<TWeakObjectPtr<AActor>>> ActorsPendingRegistration;
-#endif
 	TMap<TypedElementTableHandle, PendingRegistration<TWeakObjectPtr<UObject>>> UObjectsPendingRegistration;
 	TMap<TypedElementTableHandle, PendingRegistration<ExternalObjectRegistration>> ExternalObjectsPendingRegistration;
 	
 	TArray<ObjectRegistrationFilter> ObjectRegistrationFilters;
 	TArray<ObjectToRowDealiaser> ObjectToRowDialiasers;
+	TMap<TObjectPtr<UStruct>, TypedElementDataStorage::TableHandle> TypeToTableMap;
 	TArray<TPair<ObjectAddedCallback, FDelegateHandle>> ObjectAddedCallbackList;
 	TArray<TPair<ObjectRemovedCallback, FDelegateHandle>> PreObjectRemovedCallbackList;
 
@@ -165,9 +149,6 @@ private:
 	TypedElementTableHandle StandardUObjectTable{ TypedElementInvalidTableHandle };
 	TypedElementTableHandle StandardExternalObjectTable{ TypedElementInvalidTableHandle };
 	ITypedElementDataStorageInterface* Storage{ nullptr };
-#if TEDS_SEPARATE_ACTOR_REGISTRATION
-	TSharedPtr<FMassActorManager> ActorSubsystem;
-#endif
 
 	/**
 	 * Reference of objects (UObject and AActor) that need to be fully synced from the world to the database.
diff --git a/Engine/Plugins/Experimental/TypedElementsDataStorage/Source/TypedElementsDataStorage/Private/TypedElementsDataStorage.cpp b/Engine/Plugins/Experimental/TypedElementsDataStorage/Source/TypedElementsDataStorage/Private/TypedElementsDataStorage.cpp
--- a/Engine/Plugins/Experimental/TypedElementsDataStorage/Source/TypedElementsDataStorage/Private/TypedElementsDataStorage.cpp
+++ b/Engine/Plugins/Experimental/TypedElementsDataStorage/Source/TypedElementsDataStorage/Private/TypedElementsDataStorage.cpp
@@ -124,6 +124,7 @@ void FTypedElementsDataStorageModule::StartupModule()
 				for (UTypedElementDataStorageFactory* Factory : Factories)
 				{
 					Factory->RegisterTables(*Database);
+					Factory->RegisterTables(*Database, *DatabaseCompatibility);
 					Factory->RegisterTickGroups(*Database);
 					Factory->RegisterRegistrationFilters(*DatabaseCompatibility);
 					Factory->RegisterDealiaser(*DatabaseCompatibility);
diff --git a/Engine/Source/Runtime/TypedElementFramework/Public/Elements/Columns/TypedElementCompatibilityColumns.h b/Engine/Source/Runtime/TypedElementFramework/Public/Elements/Columns/TypedElementCompatibilityColumns.h
--- a/Engine/Source/Runtime/TypedElementFramework/Public/Elements/Columns/TypedElementCompatibilityColumns.h
+++ b/Engine/Source/Runtime/TypedElementFramework/Public/Elements/Columns/TypedElementCompatibilityColumns.h
@@ -34,14 +34,3 @@ struct FTypedElementExternalObjectColumn final : public FTypedElementDataStorage
 
 	void* Object;
 };
-
-/**
- * Source table used to register the external object with.
- */
-USTRUCT(meta = (DisplayName = "Object source table"))
-struct FTypedElementObjectSourceTableColumn final : public FTypedElementDataStorageColumn
-{
-	GENERATED_BODY()
-
-	TypedElementDataStorage::TableHandle SourceTable;
-};
\ No newline at end of file
diff --git a/Engine/Source/Runtime/TypedElementFramework/Public/Elements/Interfaces/TypedElementDataStorageCompatibilityInterface.h b/Engine/Source/Runtime/TypedElementFramework/Public/Elements/Interfaces/TypedElementDataStorageCompatibilityInterface.h
--- a/Engine/Source/Runtime/TypedElementFramework/Public/Elements/Interfaces/TypedElementDataStorageCompatibilityInterface.h
+++ b/Engine/Source/Runtime/TypedElementFramework/Public/Elements/Interfaces/TypedElementDataStorageCompatibilityInterface.h
@@ -13,7 +13,6 @@
 
 #include "TypedElementDataStorageCompatibilityInterface.generated.h"
 
-class AActor;
 struct FTypedElementDatabaseCompatibilityObjectTypeInfo;
 
 UINTERFACE(MinimalAPI)
@@ -48,9 +47,7 @@ public:
 	 */
 	template<typename ObjectType>
 	TypedElementRowHandle AddCompatibleObject(ObjectType&& Object);
-	template<typename ObjectType>
-	TypedElementRowHandle AddCompatibleObject(ObjectType&& Object, TypedElementTableHandle Table);
-
+	
 	/** Removes a previously registered object from the data storage. */
 	template<typename ObjectType>
 	void RemoveCompatibleObject(ObjectType&& Object);
@@ -76,6 +73,12 @@ public:
 	 * being aliased.
 	 */
 	virtual void RegisterDealiaserCallback(ObjectToRowDealiaser Dealiaser) = 0;
+	/**
+	 * Allows a specific type to be associated with a table. Whenever a compatible object is added, the type information of that object
+	 * will be used to find the closest match in the registered types and use the associated table. E.g. actors derive from uobjects so
+	 * if the type information of an actor is registered the actor table will be used instead of the uobject table.
+	 */
+	virtual void RegisterTypeTableAssociation(TObjectPtr<UStruct> TypeInfo, TypedElementDataStorage::TableHandle Table) = 0;
 
 	/**
 	 * @section Explicit functions
@@ -84,29 +87,16 @@ public:
 
 	/** Adds a UObject to the data storage. */
 	virtual TypedElementRowHandle AddCompatibleObjectExplicit(UObject* Object) = 0;
-	virtual TypedElementRowHandle AddCompatibleObjectExplicit(UObject* Object, TypedElementTableHandle Table) = 0;
-	/** Adds an actor to the data storage. */
-	virtual TypedElementRowHandle AddCompatibleObjectExplicit(AActor* Actor) = 0;
-	virtual TypedElementRowHandle AddCompatibleObjectExplicit(AActor* Actor, TypedElementTableHandle Table) = 0;
 	/** Adds an FStruct to the data storage. */
 	virtual TypedElementRowHandle AddCompatibleObjectExplicit(void* Object, TWeakObjectPtr<const UScriptStruct> TypeInfo) = 0;
-	virtual TypedElementRowHandle AddCompatibleObjectExplicit(void* Object, TWeakObjectPtr<const UScriptStruct> TypeInfo, TypedElementTableHandle Table) = 0;
-
+	
 	/** Removes a UObject from the data storage. */
 	virtual void RemoveCompatibleObjectExplicit(UObject* Object) = 0;
-	/** Removes an actor from the data storage. */
-	virtual void RemoveCompatibleObjectExplicit(AActor* Actor) = 0;
 	/** Removes an FStruct from the data storage. */
 	virtual void RemoveCompatibleObjectExplicit(void* Object) = 0;
 
 	/** Finds a previously stored UObject. If not found an invalid row handle will be returned. */
 	virtual TypedElementRowHandle FindRowWithCompatibleObjectExplicit(const UObject* Object) const = 0;
-	/** 
-	 * Finds a previously stored actor based on the object key. If not found an invalid row handle will be returned.
-	 * While FindRowWithCompatibleObject(const TObjectKey<const UObject> Object) can also be used, this call will be slightly
-	 * faster if it's already known that the target is an actor.
-	 */
-	virtual TypedElementRowHandle FindRowWithCompatibleObjectExplicit(const AActor* Actor) const = 0;
 	/** Finds a previously stored FStruct. If not found an invalid row handle will be returned. */
 	virtual TypedElementRowHandle FindRowWithCompatibleObjectExplicit(const void* Object) const = 0;
 };
@@ -136,29 +126,6 @@ TypedElementRowHandle ITypedElementDataStorageCompatibilityInterface::AddCompati
 	}
 }
 
-template<typename ObjectType>
-TypedElementRowHandle ITypedElementDataStorageCompatibilityInterface::AddCompatibleObject(ObjectType&& Object, TypedElementTableHandle Table)
-{
-	if (Table != TypedElementInvalidTableHandle)
-	{
-		auto RawPointer = GetRawPointer(Forward<ObjectType>(Object));
-		using BaseType = std::remove_cv_t<std::remove_pointer_t<decltype(RawPointer)>>;
-
-		if constexpr (std::is_base_of_v<UObject, BaseType>)
-		{
-			return AddCompatibleObjectExplicit(RawPointer, Table);
-		}
-		else
-		{
-			return AddCompatibleObjectExplicit(RawPointer, BaseType::StaticStruct(), Table);
-		}
-	}
-	else
-	{
-		return AddCompatibleObject(Object);
-	}
-}
-
 template<typename ObjectType>
 void ITypedElementDataStorageCompatibilityInterface::RemoveCompatibleObject(ObjectType&& Object)
 {
diff --git a/Engine/Source/Runtime/TypedElementFramework/Public/Elements/Interfaces/TypedElementDataStorageFactory.h b/Engine/Source/Runtime/TypedElementFramework/Public/Elements/Interfaces/TypedElementDataStorageFactory.h
--- a/Engine/Source/Runtime/TypedElementFramework/Public/Elements/Interfaces/TypedElementDataStorageFactory.h
+++ b/Engine/Source/Runtime/TypedElementFramework/Public/Elements/Interfaces/TypedElementDataStorageFactory.h
@@ -27,6 +27,7 @@ public:
 	virtual uint8 GetOrder() const { return 127; }
 
 	virtual void RegisterTables(ITypedElementDataStorageInterface& DataStorage) const {}
+	virtual void RegisterTables(ITypedElementDataStorageInterface& DataStorage, ITypedElementDataStorageCompatibilityInterface& DataStorageCompatibility) const {}
 	virtual void RegisterTickGroups(ITypedElementDataStorageInterface& DataStorage) const {}
 	virtual void RegisterQueries(ITypedElementDataStorageInterface& DataStorage) const {}
 
