Fixed double registration of objects with TEDS.

Double registration of objects caused duplicate actors to be added to TEDS when undoing the deletion of an actor. In turn the new TEDS Outliner would pick the wrong actor causing it to not behave as expected when for instance renaming the actor through another tool like the details panel.
The following changes were made:
- UObjects and regular objects are now checked to see if they're already registered. If they are, the handle to the already registered version is returned.
- The console command "TEDS.PrintActorLabels" has been extended to also print the actor's address.

diff --git a/Engine/Plugins/Experimental/TypedElementsDataStorage/Source/TypedElementsDataStorage/Private/Compatibility/TypedElementActorLabelQueries.cpp b/Engine/Plugins/Experimental/TypedElementsDataStorage/Source/TypedElementsDataStorage/Private/Compatibility/TypedElementActorLabelQueries.cpp
--- a/Engine/Plugins/Experimental/TypedElementsDataStorage/Source/TypedElementsDataStorage/Private/Compatibility/TypedElementActorLabelQueries.cpp
+++ b/Engine/Plugins/Experimental/TypedElementsDataStorage/Source/TypedElementsDataStorage/Private/Compatibility/TypedElementActorLabelQueries.cpp
@@ -32,9 +32,7 @@ FAutoConsoleCommandWithOutputDevice PrintActorLabelsConsoleCommand(
 				{
 					LabelQuery = DataStorage->RegisterQuery(
 						Select()
-							.ReadOnly<FTypedElementLabelColumn>()
-						.Where()
-							.All<FMassActorFragment>()
+							.ReadOnly<FMassActorFragment, FTypedElementLabelColumn>()
 						.Compile());
 				}
 				
@@ -43,7 +41,7 @@ FAutoConsoleCommandWithOutputDevice PrintActorLabelsConsoleCommand(
 					FString Message;
 					Output.Log(TEXT("The Typed Elements Data Storage has the following actors:"));
 					DataStorage->RunQuery(LabelQuery, CreateDirectQueryCallbackBinding(
-						[&Output, &Message](DSI::IDirectQueryContext& Context, const FTypedElementLabelColumn* Labels)
+						[&Output, &Message](DSI::IDirectQueryContext& Context, const FMassActorFragment* Actors, const FTypedElementLabelColumn* Labels)
 						{
 							const uint32 Count = Context.GetRowCount();
 
@@ -52,7 +50,12 @@ FAutoConsoleCommandWithOutputDevice PrintActorLabelsConsoleCommand(
 							// Reserve memory first to avoid repeated memory allocations.
 							for (uint32 Index = 0; Index < Count; ++Index)
 							{
-								CharacterCount += 12 /*Prefixed text size*/ + LabelsIt->Label.Len() + 1 /*Trailing new line*/;
+								CharacterCount 
+									+= 14 /* Prefixed text size */ 
+									+ 16 /* Hex address of actor */
+									+ 3 /* Closing brace, colon and space */
+									+ LabelsIt->Label.Len() 
+									+ 1 /* Trailing new line */;
 								++LabelsIt;
 							}
 							Message.Reset(CharacterCount);
@@ -60,10 +63,10 @@ FAutoConsoleCommandWithOutputDevice PrintActorLabelsConsoleCommand(
 							LabelsIt = Labels;
 							for (uint32 Index = 0; Index < Count; ++Index)
 							{
-								Message += TEXT("    Actor: ");
-								Message += LabelsIt->Label;
-								Message += TEXT('\n');
+								Message.Appendf(TEXT("    Actor (0x%p): %s\n"), Actors->Get(), *LabelsIt->Label);
+
 								++LabelsIt;
+								++Actors;
 							}
 
 							Output.Log(Message);
diff --git a/Engine/Plugins/Experimental/TypedElementsDataStorage/Source/TypedElementsDataStorage/Private/TypedElementDatabaseCompatibility.cpp b/Engine/Plugins/Experimental/TypedElementsDataStorage/Source/TypedElementsDataStorage/Private/TypedElementDatabaseCompatibility.cpp
--- a/Engine/Plugins/Experimental/TypedElementsDataStorage/Source/TypedElementsDataStorage/Private/TypedElementDatabaseCompatibility.cpp
+++ b/Engine/Plugins/Experimental/TypedElementsDataStorage/Source/TypedElementsDataStorage/Private/TypedElementDatabaseCompatibility.cpp
@@ -107,14 +107,18 @@ TypedElementRowHandle UTypedElementDatabaseCompatibility::AddCompatibleObjectExp
 	
 	if (ensureMsgf(Storage, TEXT("Trying to add an object to Typed Element's Data Storage before the storage is available.")))
 	{
-		TableHandle Table = FindBestMatchingTable(TypeInfo.Get());
-		Table = (Table != InvalidTableHandle) ? Table : StandardExternalObjectTable;
-
-		TypedElementRowHandle ReservedRow = Storage->ReserveRow();
-		Storage->IndexRow(GenerateIndexHash(Object), ReservedRow);
-		PendingRegistration<ExternalObjectRegistration>& Pending = ExternalObjectsPendingRegistration.FindOrAdd(Table);
-		Pending.Add(ReservedRow, ExternalObjectRegistration{ .Object = Object, .TypeInfo = TypeInfo });
-		return ReservedRow;
+		TypedElementRowHandle Result = FindRowWithCompatibleObjectExplicit(Object);
+		if (!Storage->IsRowAvailable(Result))
+		{
+			TableHandle Table = FindBestMatchingTable(TypeInfo.Get());
+			Table = (Table != InvalidTableHandle) ? Table : StandardExternalObjectTable;
+
+			Result = Storage->ReserveRow();
+			Storage->IndexRow(GenerateIndexHash(Object), Result);
+			PendingRegistration<ExternalObjectRegistration>& Pending = ExternalObjectsPendingRegistration.FindOrAdd(Table);
+			Pending.Add(Result, ExternalObjectRegistration{ .Object = Object, .TypeInfo = TypeInfo });
+		}
+		return Result;
 	}
 	else
 	{
@@ -248,25 +252,29 @@ TypedElementRowHandle UTypedElementDatabaseCompatibility::AddCompatibleObjectExp
 {
 	using namespace TypedElementDataStorage;
 
-	TableHandle Table = FindBestMatchingTable(Object->GetClass());
-	checkf(Table != InvalidTableHandle, TEXT("The Typed Elements Data Storage could not find any matching tables for object of type '%s'. "
-		"This can mean that the object doesn't derive from UObject or that a table for UObject is no longer registered."), *Object->GetClass()->GetFName().ToString());
+	TypedElementRowHandle Result = FindRowWithCompatibleObjectExplicit(Object);
+	if (!Storage->IsRowAvailable(Result))
+	{
+		TableHandle Table = FindBestMatchingTable(Object->GetClass());
+		checkf(Table != InvalidTableHandle, TEXT("The Typed Elements Data Storage could not find any matching tables for object of type '%s'. "
+			"This can mean that the object doesn't derive from UObject or that a table for UObject is no longer registered."), *Object->GetClass()->GetFName().ToString());
 
-	TypedElementRowHandle ReservedRow = Storage->ReserveRow();
-	Storage->IndexRow(GenerateIndexHash(Object), ReservedRow);
+		Result = Storage->ReserveRow();
+		Storage->IndexRow(GenerateIndexHash(Object), Result);
 
-	PendingRegistration<TWeakObjectPtr<UObject>>& Pending = UObjectsPendingRegistration.FindOrAdd(Table);
-	Pending.Add(ReservedRow, Object);
+		PendingRegistration<TWeakObjectPtr<UObject>>& Pending = UObjectsPendingRegistration.FindOrAdd(Table);
+		Pending.Add(Result, Object);
 
-	if constexpr (bEnableTransactions)
-	{
-		if (GUndo)
+		if constexpr (bEnableTransactions)
 		{
-			GUndo->StoreUndo(this, MakeUnique<FRegistrationCommandChange>(Object));
+			if (GUndo)
+			{
+				GUndo->StoreUndo(this, MakeUnique<FRegistrationCommandChange>(Object));
+			}
 		}
 	}
 
-	return ReservedRow;
+	return Result;
 }
 
 template<bool bEnableTransactions>
