Fixed a large number of objects being registered with TEDS after a script recompile.

Whenever scripts recompiles a large number of object come through TEDS' reinstancing manager. A previous change added code that reregisters actors, but this has the unexpected side effect of now also registering everything that it send over during the recompile. The fix is to only reregister objects that have a memento recorded for them. This implies that the object was previously seen in TEDS and that there is still an interest in it, meaning any other objects can be ignored.

diff --git a/Engine/Plugins/Experimental/TypedElementsDataStorage/Source/TypedElementsDataStorage/Private/Compatibility/TypedElementObjectReinstancingManager.cpp b/Engine/Plugins/Experimental/TypedElementsDataStorage/Source/TypedElementsDataStorage/Private/Compatibility/TypedElementObjectReinstancingManager.cpp
--- a/Engine/Plugins/Experimental/TypedElementsDataStorage/Source/TypedElementsDataStorage/Private/Compatibility/TypedElementObjectReinstancingManager.cpp
+++ b/Engine/Plugins/Experimental/TypedElementsDataStorage/Source/TypedElementsDataStorage/Private/Compatibility/TypedElementObjectReinstancingManager.cpp
@@ -92,18 +92,18 @@ void UTypedElementObjectReinstancingManager::HandleOnObjectsReinstanced(
 {
 	for (FCoreUObjectDelegates::FReplacementObjectMap::TConstIterator Iter = ObjectReplacementMap.CreateConstIterator(); Iter; ++Iter)
 	{
-		UObject* NewInstanceObject = Iter->Value;
-
-		TypedElementRowHandle NewObjectRow = DataStorageCompatibility->FindRowWithCompatibleObjectExplicit(NewInstanceObject);
-		if (!Database->IsRowAvailable(NewObjectRow))
-		{
-			NewObjectRow = DataStorageCompatibility->AddCompatibleObjectExplicit(NewInstanceObject);
-		}
-		
 		const void* PreDeleteObject = Iter->Key;
-		const TypedElementRowHandle* MementoRowPtr = OldObjectToMementoMap.Find(PreDeleteObject);
-		if (MementoRowPtr != nullptr)
+		if (const TypedElementRowHandle* MementoRowPtr = OldObjectToMementoMap.Find(PreDeleteObject))
 		{
+			UObject* NewInstanceObject = Iter->Value;
+			TypedElementRowHandle NewObjectRow = DataStorageCompatibility->FindRowWithCompatibleObjectExplicit(NewInstanceObject);
+			// Do the addition only if there's a recorded memento. Having a memento implies the object was previously registered and there's
+			// still an interest in it. Any other objects can therefore be ignored.
+			if (!Database->IsRowAvailable(NewObjectRow))
+			{
+				NewObjectRow = DataStorageCompatibility->AddCompatibleObjectExplicit(NewInstanceObject);
+			}
+
 			// Kick off re-instantiation of NewObjectRow from the Memento
 			TypedElementRowHandle Memento = *MementoRowPtr;
 			if (ensureMsgf(Database->HasColumns(Memento, TConstArrayView<const UScriptStruct*>({FTypedElementMementoTag::StaticStruct()})), TEXT("Cannot reinstantiate from a non memento row")))
