TEDS - Add a TypedElementUIStyleOverrideCapability to allow specifying ColorAndOpacity to TypedElementLabelWidget alongside a IsReadOnly attribute on TypedElementUIEditableCapability to disallow entering edit mode

diff --git a/Engine/Plugins/Experimental/TypedElementsDataStorage/Source/TypedElementsDataStorageUI/Private/TypedElementLabelWidget.cpp b/Engine/Plugins/Experimental/TypedElementsDataStorage/Source/TypedElementsDataStorageUI/Private/TypedElementLabelWidget.cpp
--- a/Engine/Plugins/Experimental/TypedElementsDataStorage/Source/TypedElementsDataStorageUI/Private/TypedElementLabelWidget.cpp
+++ b/Engine/Plugins/Experimental/TypedElementsDataStorage/Source/TypedElementsDataStorageUI/Private/TypedElementLabelWidget.cpp
@@ -11,6 +11,7 @@
 #include "Elements/Interfaces/Capabilities/TypedElementUiEditableCapability.h"
 #include "Elements/Interfaces/Capabilities/TypedElementUiTextCapability.h"
 #include "Elements/Interfaces/Capabilities/TypedElementUiTooltipCapability.h"
+#include "Elements/Interfaces/Capabilities/TypedElementUiStyleOverrideCapability.h"
 #include "TypedElementSubsystems.h"
 #include "Widgets/Text/STextBlock.h"
 #include "Widgets/Text/SInlineEditableTextBlock.h"
@@ -195,6 +196,7 @@ TSharedPtr<SWidget> FTypedElementLabelWidgetConstructor::Construct(
 			TextBlock->AddMetadata(MakeShared<TTypedElementUiEditableCapability<SInlineEditableTextBlock>>(*TextBlock));
 			TextBlock->AddMetadata(MakeShared<TTypedElementUiTextCapability<SInlineEditableTextBlock>>(*TextBlock));
 			TextBlock->AddMetadata(MakeShared<TTypedElementUiTooltipCapability<SInlineEditableTextBlock>>(*TextBlock));
+			TextBlock->AddMetadata(MakeShared<TTypedElementUiStyleOverrideCapability<SInlineEditableTextBlock>>(*TextBlock));
 			Result = TextBlock;
 		}
 	}
diff --git a/Engine/Source/Runtime/Slate/Private/Widgets/Text/SInlineEditableTextBlock.cpp b/Engine/Source/Runtime/Slate/Private/Widgets/Text/SInlineEditableTextBlock.cpp
--- a/Engine/Source/Runtime/Slate/Private/Widgets/Text/SInlineEditableTextBlock.cpp
+++ b/Engine/Source/Runtime/Slate/Private/Widgets/Text/SInlineEditableTextBlock.cpp
@@ -159,7 +159,7 @@ bool SInlineEditableTextBlock::IsInEditMode() const
 	return TextBlock->GetVisibility() == EVisibility::Collapsed;
 }
 
-void SInlineEditableTextBlock::SetReadOnly(bool bInIsReadOnly)
+void SInlineEditableTextBlock::SetReadOnly(const TAttribute<bool>& bInIsReadOnly)
 {
 	bIsReadOnly = bInIsReadOnly;
 }
diff --git a/Engine/Source/Runtime/Slate/Public/Widgets/Text/SInlineEditableTextBlock.h b/Engine/Source/Runtime/Slate/Public/Widgets/Text/SInlineEditableTextBlock.h
--- a/Engine/Source/Runtime/Slate/Public/Widgets/Text/SInlineEditableTextBlock.h
+++ b/Engine/Source/Runtime/Slate/Public/Widgets/Text/SInlineEditableTextBlock.h
@@ -153,7 +153,7 @@ class SInlineEditableTextBlock: public SCompoundWidget
 	/** Checks if the widget is in edit mode */
 	SLATE_API bool IsInEditMode() const;
 
-	SLATE_API void SetReadOnly(bool bInIsReadOnly);
+	SLATE_API void SetReadOnly(const TAttribute<bool>& bInIsReadOnly);
 
 	SLATE_API void SetText( const TAttribute< FText >& InText );
 	SLATE_API void SetText( const FString& InText );
diff --git a/Engine/Source/Runtime/TypedElementFramework/Public/Elements/Interfaces/Capabilities/TypedElementUiEditableCapability.h b/Engine/Source/Runtime/TypedElementFramework/Public/Elements/Interfaces/Capabilities/TypedElementUiEditableCapability.h
--- a/Engine/Source/Runtime/TypedElementFramework/Public/Elements/Interfaces/Capabilities/TypedElementUiEditableCapability.h
+++ b/Engine/Source/Runtime/TypedElementFramework/Public/Elements/Interfaces/Capabilities/TypedElementUiEditableCapability.h
@@ -3,6 +3,7 @@
 #pragma once
 
 #include "Elements/Interfaces/TypedElementUiCapabilities.h"
+#include "Misc/Attribute.h"
 
 /**
  * Interface to provide access to widget that have explicit editing modes.
@@ -16,6 +17,7 @@ public:
 
 	virtual void EnterEditingMode() = 0;
 	virtual void ExitEditingMode() = 0;
+	virtual void SetIsReadOnly(const TAttribute<bool>& InIsReadOnly) = 0;
 };
 
 template<typename WidgetType>
@@ -32,6 +34,11 @@ public:
 	{
 		Widget.ExitEditingMode();
 	}
+	
+	void SetIsReadOnly(const TAttribute<bool>& InIsReadOnly) override
+	{
+		Widget.SetReadOnly(InIsReadOnly);
+	}
 
 private:
 	WidgetType& Widget;
diff --git a/Engine/Source/Runtime/TypedElementFramework/Public/Elements/Interfaces/Capabilities/TypedElementUiStyleOverrideCapability.h b/Engine/Source/Runtime/TypedElementFramework/Public/Elements/Interfaces/Capabilities/TypedElementUiStyleOverrideCapability.h
new file mode 100644
--- /dev/null
+++ b/Engine/Source/Runtime/TypedElementFramework/Public/Elements/Interfaces/Capabilities/TypedElementUiStyleOverrideCapability.h
@@ -0,0 +1,35 @@
+// Copyright Epic Games, Inc. All Rights Reserved.
+
+#pragma once
+
+#include "Elements/Interfaces/TypedElementUiCapabilities.h"
+#include "Misc/Attribute.h"
+#include "Math/Color.h"
+
+/**
+ * Interface to provide access to widgets that support working with style overrides.
+ */
+class ITypedElementUiStyleOverrideCapability : public ITypedElementUiCapability
+{
+public:
+	SLATE_METADATA_TYPE(ITypedElementUiStyleOverrideCapability, ITypedElementUiCapability)
+
+	~ITypedElementUiStyleOverrideCapability() override = default;
+
+	virtual void SetColorAndOpacity(const TAttribute<FLinearColor>& InColorAndOpacity) = 0;
+};
+
+template<typename WidgetType>
+class TTypedElementUiStyleOverrideCapability : public ITypedElementUiStyleOverrideCapability
+{
+public:
+	explicit TTypedElementUiStyleOverrideCapability(WidgetType& InWidget) : Widget(InWidget){}
+	
+	void SetColorAndOpacity(const TAttribute<FLinearColor>& InColorAndOpacity) override
+	{
+		Widget.SetColorAndOpacity(InColorAndOpacity);
+	}
+
+private:
+	WidgetType& Widget;
+};
\ No newline at end of file
diff --git a/Engine/Source/Runtime/TypedElementFramework/Public/Elements/Interfaces/Capabilities/TypedElementUiTextCapability.h b/Engine/Source/Runtime/TypedElementFramework/Public/Elements/Interfaces/Capabilities/TypedElementUiTextCapability.h
--- a/Engine/Source/Runtime/TypedElementFramework/Public/Elements/Interfaces/Capabilities/TypedElementUiTextCapability.h
+++ b/Engine/Source/Runtime/TypedElementFramework/Public/Elements/Interfaces/Capabilities/TypedElementUiTextCapability.h
@@ -16,9 +16,9 @@ public:
 
 	~ITypedElementUiTextCapability() override = default;
 
-	virtual void SetText(TAttribute<FText> Text) = 0;
+	virtual void SetText(const TAttribute<FText>& Text) = 0;
 
-	virtual void SetHighlightText(TAttribute<FText> Text) = 0;
+	virtual void SetHighlightText(const TAttribute<FText>& Text) = 0;
 };
 
 template<typename WidgetType>
@@ -27,12 +27,12 @@ class TTypedElementUiTextCapability : public ITypedElementUiTextCapability
 public:
 	explicit TTypedElementUiTextCapability(WidgetType& InWidget) : Widget(InWidget){}
 	
-	void SetText(TAttribute<FText> Text) override
+	void SetText(const TAttribute<FText>& Text) override
 	{
 		Widget.SetText(Text);
 	}
 
-	void SetHighlightText(TAttribute<FText> Text) override
+	void SetHighlightText(const TAttribute<FText>& Text) override
 	{
 		Widget.SetHighlightText(Text);
 	}
