Fix crash that can occur in TEDS reinstancing manager if reinstance map provides null objects
* Added FTypedElementMementoReinstanceAborted tag that will also match the cleanup processor: "Delete mementos used for reinstancing" to avoid a conditional in the mementoization processors
* When a mapping does not exists, the ReinstanceAborted tag is added instead of the ReinstanceTarget.
* Use slightly more efficient row delete method in cleanup processor
* Pushed ensure for the memento tag up to where mementos are created.

diff --git a/Engine/Plugins/Experimental/TypedElementsDataStorage/Source/TypedElementsDataStorage/Private/Compatibility/TypedElementObjectReinstancingManager.cpp b/Engine/Plugins/Experimental/TypedElementsDataStorage/Source/TypedElementsDataStorage/Private/Compatibility/TypedElementObjectReinstancingManager.cpp
--- a/Engine/Plugins/Experimental/TypedElementsDataStorage/Source/TypedElementsDataStorage/Private/Compatibility/TypedElementObjectReinstancingManager.cpp
+++ b/Engine/Plugins/Experimental/TypedElementsDataStorage/Source/TypedElementsDataStorage/Private/Compatibility/TypedElementObjectReinstancingManager.cpp
@@ -80,6 +80,14 @@ void UTypedElementObjectReinstancingManager::HandleOnObjectPreRemoved(const void
 {
 	// This is the chance to record the old object to memento
 	TypedElementRowHandle Memento = MementoSystem->CreateMemento(Database.Get());
+	
+	if (!ensureMsgf(Database->HasColumns(Memento, TConstArrayView<const UScriptStruct*>({FTypedElementMementoTag::StaticStruct()})),
+		TEXT("Memento System should create rows with the MementoTag")))
+	{
+		Database->RemoveRow(Memento);
+		return;
+	}
+	
 	Database->AddOrGetColumn(ObjectRow, FTypedElementMementoOnDelete{ .Memento = Memento });
 	
 	Database->AddOrGetColumn(Memento, FTypedElementsReinstanceableSourceObject{.Object = Object});
@@ -95,7 +103,16 @@ void UTypedElementObjectReinstancingManager::HandleOnObjectsReinstanced(
 		const void* PreDeleteObject = Iter->Key;
 		if (const TypedElementRowHandle* MementoRowPtr = OldObjectToMementoMap.Find(PreDeleteObject))
 		{
+			TypedElementRowHandle Memento = *MementoRowPtr;
+			
 			UObject* NewInstanceObject = Iter->Value;
+			if (NewInstanceObject == nullptr)
+			{
+				// Reinstancing resulted in no target object.  Delete the memento.
+				Database->AddOrGetColumn<FTypedElementMementoReinstanceAborted>(Memento);
+				continue;
+			}
+			
 			TypedElementRowHandle NewObjectRow = DataStorageCompatibility->FindRowWithCompatibleObjectExplicit(NewInstanceObject);
 			// Do the addition only if there's a recorded memento. Having a memento implies the object was previously registered and there's
 			// still an interest in it. Any other objects can therefore be ignored.
@@ -105,11 +122,7 @@ void UTypedElementObjectReinstancingManager::HandleOnObjectsReinstanced(
 			}
 
 			// Kick off re-instantiation of NewObjectRow from the Memento
-			TypedElementRowHandle Memento = *MementoRowPtr;
-			if (ensureMsgf(Database->HasColumns(Memento, TConstArrayView<const UScriptStruct*>({FTypedElementMementoTag::StaticStruct()})), TEXT("Cannot reinstantiate from a non memento row")))
-			{
-				Database->AddOrGetColumn(Memento, FTypedElementMementoReinstanceTarget{ .Target = NewObjectRow });
-			}
+			Database->AddOrGetColumn(Memento, FTypedElementMementoReinstanceTarget{ .Target = NewObjectRow });
 		}
 	}
 }
diff --git a/Engine/Plugins/Experimental/TypedElementsDataStorage/Source/TypedElementsDataStorage/Private/Memento/TypedElementMementoRowTypes.h b/Engine/Plugins/Experimental/TypedElementsDataStorage/Source/TypedElementsDataStorage/Private/Memento/TypedElementMementoRowTypes.h
--- a/Engine/Plugins/Experimental/TypedElementsDataStorage/Source/TypedElementsDataStorage/Private/Memento/TypedElementMementoRowTypes.h
+++ b/Engine/Plugins/Experimental/TypedElementsDataStorage/Source/TypedElementsDataStorage/Private/Memento/TypedElementMementoRowTypes.h
@@ -30,3 +30,14 @@ struct FTypedElementMementoReinstanceTarget : public FTypedElementDataStorageCol
 	GENERATED_BODY()
 	TypedElementRowHandle Target;
 };
+
+/**
+ * Tag added to memento row which will cause the memento to be deleted.  It is used
+ * in the case where mementoization was aborted due to a missing target object
+ * from the reinstancing callback
+ */
+USTRUCT()
+struct FTypedElementMementoReinstanceAborted : public FTypedElementDataStorageTag
+{
+	GENERATED_BODY()
+};
\ No newline at end of file
diff --git a/Engine/Plugins/Experimental/TypedElementsDataStorage/Source/TypedElementsDataStorage/Private/Memento/TypedElementMementoSystem.cpp b/Engine/Plugins/Experimental/TypedElementsDataStorage/Source/TypedElementsDataStorage/Private/Memento/TypedElementMementoSystem.cpp
--- a/Engine/Plugins/Experimental/TypedElementsDataStorage/Source/TypedElementsDataStorage/Private/Memento/TypedElementMementoSystem.cpp
+++ b/Engine/Plugins/Experimental/TypedElementsDataStorage/Source/TypedElementsDataStorage/Private/Memento/TypedElementMementoSystem.cpp
@@ -168,21 +168,23 @@ void UTypedElementMementoSystem::RegisterQueries(UTypedElementDatabase& DataStor
 				.Compile());
 		check(QueryHandleForExistence != TypedElementInvalidQueryHandle);
 	}
-
+	
 	/**
-	 * A processor which deletes mementos that were used for rehydrating another row
+	 * A processor which deletes mementos that were used for reinstancing another row
 	 */
 	{
 		const TypedElementQueryHandle QueryHandle = DataStorage.RegisterQuery(
 			Select(
 			TEXT("Delete mementos used for reinstancing"),
 			FProcessor(DSI::EQueryTickPhase::FrameEnd, DataStorage.GetQueryTickGroupName(DSI::EQueryTickGroups::Default)),
-				[](TypedElementDataStorage::IQueryContext& Context, TypedElementRowHandle Row)
+			[](TypedElementDataStorage::IQueryContext& Context, const TypedElementRowHandle* Rows)
 				{
-					Context.RemoveRow(Row);
+					TConstArrayView<TypedElementRowHandle> RowView(Rows, Context.GetRowCount());
+					Context.RemoveRows(RowView);
 				})
-				.ReadOnly<FTypedElementMementoReinstanceTarget>()
-				.Where().All<FTypedElementMementoTag>()
+				.Where()
+					.All<FTypedElementMementoTag>()
+					.Any<FTypedElementMementoReinstanceTarget, FTypedElementMementoReinstanceAborted>()
 				.Compile()
 			);
 		check(QueryHandle != TypedElementInvalidQueryHandle);
