Add option to automatically resolve revision control state into TEDS

diff --git a/Engine/Plugins/Experimental/TypedElementsDataStorage/Source/TypedElementsDataStorage/Private/Compatibility/TypedElementSCCProcessors.cpp b/Engine/Plugins/Experimental/TypedElementsDataStorage/Source/TypedElementsDataStorage/Private/Compatibility/TypedElementSCCProcessors.cpp
deleted file mode 100644
diff --git a/Engine/Plugins/Experimental/TypedElementsDataStorage/Source/TypedElementsDataStorage/Private/Compatibility/TypedElementSCCProcessors.h b/Engine/Plugins/Experimental/TypedElementsDataStorage/Source/TypedElementsDataStorage/Private/Compatibility/TypedElementSCCProcessors.h
deleted file mode 100644
diff --git a/Engine/Plugins/Experimental/TypedElementsDataStorage/Source/TypedElementsDataStorage/Private/Compatibility/TypedElementUObjectPackagePathToColumnQueries.cpp b/Engine/Plugins/Experimental/TypedElementsDataStorage/Source/TypedElementsDataStorage/Private/Compatibility/TypedElementUObjectPackagePathToColumnQueries.cpp
--- a/Engine/Plugins/Experimental/TypedElementsDataStorage/Source/TypedElementsDataStorage/Private/Compatibility/TypedElementUObjectPackagePathToColumnQueries.cpp
+++ b/Engine/Plugins/Experimental/TypedElementsDataStorage/Source/TypedElementsDataStorage/Private/Compatibility/TypedElementUObjectPackagePathToColumnQueries.cpp
@@ -23,7 +23,7 @@ static void ResolvePackageReference(ITypedElementDataStorageInterface::IQueryCon
 	LoadedPathColumn.LoadedPath = Package->GetLoadedPath();
 				
 	Context.AddColumn(PackageRow, MoveTemp(PathColumn));
-	Context.AddColumn(PackageRow, MoveTemp(LoadedPathColumn));	
+	Context.AddColumn(PackageRow, MoveTemp(LoadedPathColumn));
 };
 
 void UTypedElementUObjectPackagePathFactory::RegisterQueries(ITypedElementDataStorageInterface& DataStorage) const
@@ -46,7 +46,8 @@ void UTypedElementUObjectPackagePathFactory::RegisterQueries(ITypedElementDataSt
 					if (FString PackageFilename; FPackageName::TryConvertLongPackageNameToFilename(Path, PackageFilename))
 					{
 						FPaths::NormalizeFilename(PackageFilename);
-						TypedElementDataStorage::IndexHash Index = TypedElementDataStorage::GenerateIndexHash(FPaths::ConvertRelativePathToFull(PackageFilename));
+						FString FullPackageFilename = FPaths::ConvertRelativePathToFull(PackageFilename);
+						TypedElementDataStorage::IndexHash Index = TypedElementDataStorage::GenerateIndexHash(FullPackageFilename);
 						TypedElementRowHandle PackageRow = Context.FindIndexedRow(Index);
 						if (Context.IsRowAvailable(PackageRow))
 						{
@@ -57,6 +58,7 @@ void UTypedElementUObjectPackagePathFactory::RegisterQueries(ITypedElementDataSt
 						{
 							FTypedElementPackageUnresolvedReference UnresolvedPackageReference;
                             UnresolvedPackageReference.Index = Index;
+							UnresolvedPackageReference.PathOnDisk = MoveTemp(FullPackageFilename);
                             Context.AddColumn(Row, MoveTemp(UnresolvedPackageReference));
 						}
 					}
@@ -69,7 +71,7 @@ void UTypedElementUObjectPackagePathFactory::RegisterQueries(ITypedElementDataSt
 	DataStorage.RegisterQuery(
 		Select(
 			TEXT("Resolve package references"),
-			FProcessor(DSI::EQueryTickPhase::DuringPhysics, DataStorage.GetQueryTickGroupName(DSI::EQueryTickGroups::SyncExternalToDataStorage)),
+			FProcessor(DSI::EQueryTickPhase::FrameEnd, DataStorage.GetQueryTickGroupName(DSI::EQueryTickGroups::SyncExternalToDataStorage)),
 			[](DSI::IQueryContext& Context, TypedElementRowHandle Row, const FTypedElementUObjectColumn& Object, const FTypedElementPackageUnresolvedReference& UnresolvedPackageReference)
 			{
 				TypedElementRowHandle PackageRow = Context.FindIndexedRow(UnresolvedPackageReference.Index);
diff --git a/Engine/Plugins/Experimental/TypedElementsDataStorage/Source/TypedElementsDataStorageRevisionControl/Source/Private/RevisionControlProcessors.cpp b/Engine/Plugins/Experimental/TypedElementsDataStorage/Source/TypedElementsDataStorageRevisionControl/Source/Private/RevisionControlProcessors.cpp
new file mode 100644
--- /dev/null
+++ b/Engine/Plugins/Experimental/TypedElementsDataStorage/Source/TypedElementsDataStorageRevisionControl/Source/Private/RevisionControlProcessors.cpp
@@ -0,0 +1,172 @@
+// Copyright Epic Games, Inc. All Rights Reserved.
+
+#include "RevisionControlProcessors.h"
+
+#include "ISourceControlModule.h"
+#include "SourceControlFileStatusMonitor.h"
+#include "HAL/IConsoleManager.h"
+
+#include "Elements/Columns/TypedElementMiscColumns.h"
+#include "Elements/Columns/TypedElementPackageColumns.h"
+#include "Elements/Columns/TypedElementRevisionControlColumns.h"
+#include "Elements/Columns/TypedElementSelectionColumns.h"
+#include "Elements/Columns/TypedElementViewportColumns.h"
+#include "Elements/Framework/TypedElementQueryBuilder.h"
+#include "Elements/Framework/TypedElementRegistry.h"
+
+static bool bAutoPopulateRevisionControlState = false;
+static FAutoConsoleVariableRef CVarAutoPopulateState(
+	TEXT("TEDS.RevisionControl.AutoPopulateState"),
+	bAutoPopulateRevisionControlState,
+	TEXT("Automatically query revision control provider and fill information into TEDS")
+);
+
+FAutoConsoleCommandWithArgsAndOutputDevice SetSelectionSCCStateConsoleCommand(
+	TEXT("TEDS.Debug.SetSCCState"),
+	TEXT("Adds a source control state to selected objects."),
+	FConsoleCommandWithArgsAndOutputDeviceDelegate::CreateLambda([](const TArray<FString>& Args, FOutputDevice& Output)
+		{
+			using namespace TypedElementQueryBuilder;
+			using DSI = ITypedElementDataStorageInterface;
+
+			TRACE_CPUPROFILER_EVENT_SCOPE(TEDS.Debug.SetSCCState);
+
+			if (ITypedElementDataStorageInterface* DataStorage = UTypedElementRegistry::GetInstance()->GetMutableDataStorage())
+			{
+				static TypedElementQueryHandle AllSelectedQuery = TypedElementInvalidQueryHandle;
+				if (AllSelectedQuery == TypedElementInvalidQueryHandle)
+				{
+					AllSelectedQuery = DataStorage->RegisterQuery(
+						Select()
+						.Where()
+							.All<FTypedElementSelectionColumn>()
+						.Compile());
+				}
+				
+				if (AllSelectedQuery == TypedElementInvalidQueryHandle)
+				{
+					return;
+				}
+
+				ESCCModification Modification = ESCCModification::Modified;
+
+				if (Args.Num() > 0)
+				{
+					// Parse the index
+					int32 StateIndex;
+					LexFromString(StateIndex, *Args[0]);
+
+					if (!(StateIndex >= 0 && StateIndex <= int32(ESCCModification::Conflicted)))
+					{
+						Output.Log(TEXT("State index out of range"));
+						return;
+					}
+
+					Modification = static_cast<ESCCModification>(StateIndex);
+				}
+				
+				TArray<TypedElementRowHandle> RowHandles;
+				
+				DataStorage->RunQuery(AllSelectedQuery, [&RowHandles](const DSI::FQueryDescription&, DSI::IDirectQueryContext& Context)
+				{
+					RowHandles = Context.GetRowHandles();
+				});
+
+				for (TypedElementRowHandle Row : RowHandles)
+				{
+					DataStorage->AddOrGetColumn<FSCCStatusColumn>(Row)->Modification = Modification;
+					DataStorage->AddColumn<FTypedElementSyncFromWorldTag>(Row);
+				}
+			}
+		}
+	));
+
+void UTypedElementRevisionControlFactory::RegisterTables(ITypedElementDataStorageInterface& DataStorage) const
+{
+	DataStorage.RegisterTable(
+		TTypedElementColumnTypeList<
+			FTypedElementPackagePathColumn, FTypedElementPackageLoadedPathColumn,
+			FSCCRevisionIdColumn, FSCCExternalRevisionIdColumn>(),
+		FName("Editor_RevisionControlTable"));
+}
+
+void UTypedElementRevisionControlFactory::RegisterQueries(ITypedElementDataStorageInterface& DataStorage) const
+{
+	using namespace TypedElementQueryBuilder;
+	using DSI = ITypedElementDataStorageInterface;
+
+	DataStorage.RegisterQuery(
+		Select(
+			TEXT("Change selection outline colors based on SCC status"),
+			// This is in PrePhysics because the outline->actor query is in DuringPhysics and contexts don't flush changes between tick groups
+			FProcessor(DSI::EQueryTickPhase::PrePhysics, DataStorage.GetQueryTickGroupName(DSI::EQueryTickGroups::SyncExternalToDataStorage)),
+			[](DSI::IQueryContext& Context, TypedElementRowHandle Row, const FSCCStatusColumn& SCCStatus)
+			{
+				constexpr int BasicSelectionColorCount = 2;
+				Context.AddColumn<FTypedElementViewportColorColumn>(Row, { .SelectionOutlineColorIndex = static_cast<uint8>(int(SCCStatus.Modification) + BasicSelectionColorCount) });
+				Context.AddColumns<FTypedElementSyncBackToWorldTag>(Row);
+			}
+		)
+		.Where()
+			.All<FTypedElementSyncFromWorldTag>()
+		.Compile()
+	);
+
+	CVarAutoPopulateState->AsVariable()->SetOnChangedCallback(
+		FConsoleVariableDelegate::CreateLambda([this, &DataStorage](IConsoleVariable* AutoPopulate)
+		{
+			if (AutoPopulate->GetBool())
+			{
+				RegisterFetchUpdates(DataStorage);
+			}
+			else
+			{
+				DataStorage.UnregisterQuery(FetchUpdates);
+			}
+		})
+	); 
+	
+	if (bAutoPopulateRevisionControlState)
+	{
+		RegisterFetchUpdates(DataStorage);
+	}
+}
+
+void UTypedElementRevisionControlFactory::RegisterFetchUpdates(ITypedElementDataStorageInterface& DataStorage) const
+{
+	using namespace TypedElementQueryBuilder;
+	using DSI = ITypedElementDataStorageInterface;
+	
+	FSourceControlFileStatusMonitor& FileStatusMonitor = ISourceControlModule::Get().GetSourceControlFileStatusMonitor();
+
+	FetchUpdates = DataStorage.RegisterQuery(
+		Select(
+			TEXT("Gather source control statuses for objects with unresolved package paths"),
+			FProcessor(DSI::EQueryTickPhase::DuringPhysics, DataStorage.GetQueryTickGroupName(DSI::EQueryTickGroups::SyncExternalToDataStorage))
+				.ForceToGameThread(true),
+			[this, &FileStatusMonitor](DSI::IQueryContext& Context, const FTypedElementPackageUnresolvedReference* InUnresolvedReferences)
+			{
+				TConstArrayView<TypedElementDataStorage::RowHandle> RowHandles = Context.GetRowHandles();
+				TConstArrayView<FTypedElementPackageUnresolvedReference, int64> UnresolvedReferences { InUnresolvedReferences, Context.GetRowCount() };
+
+				for (int64 UnresolvedReferenceIndex = 0; UnresolvedReferenceIndex < UnresolvedReferences.Num(); ++UnresolvedReferenceIndex)
+				{
+					const FTypedElementPackageUnresolvedReference& UnresolvedReference = UnresolvedReferences[UnresolvedReferenceIndex];
+					if (UnresolvedReference.Index == 0)
+					{
+						Context.RemoveColumns<FTypedElementPackageUnresolvedReference>(RowHandles[UnresolvedReferenceIndex]);
+						return;
+					}
+					static FSourceControlFileStatusMonitor::FOnSourceControlFileStatus EmptyDelegate{};
+					
+					FileStatusMonitor.StartMonitoringFile(
+						reinterpret_cast<uintptr_t>(this),
+						UnresolvedReference.PathOnDisk,
+						EmptyDelegate
+					);
+				}
+			}
+		)
+		.Compile()
+	);
+}
diff --git a/Engine/Plugins/Experimental/TypedElementsDataStorage/Source/TypedElementsDataStorageRevisionControl/Source/Private/RevisionControlProcessors.h b/Engine/Plugins/Experimental/TypedElementsDataStorage/Source/TypedElementsDataStorageRevisionControl/Source/Private/RevisionControlProcessors.h
new file mode 100644
--- /dev/null
+++ b/Engine/Plugins/Experimental/TypedElementsDataStorage/Source/TypedElementsDataStorageRevisionControl/Source/Private/RevisionControlProcessors.h
@@ -0,0 +1,25 @@
+// Copyright Epic Games, Inc. All Rights Reserved.
+
+#pragma once
+
+#include "Elements/Interfaces/TypedElementDataStorageFactory.h"
+#include "Elements/Interfaces/TypedElementQueryStorageInterfaces.h"
+#include "UObject/ObjectMacros.h"
+
+#include "RevisionControlProcessors.generated.h"
+
+UCLASS()
+class UTypedElementRevisionControlFactory : public UTypedElementDataStorageFactory
+{
+	GENERATED_BODY()
+
+public:
+	~UTypedElementRevisionControlFactory() override = default;
+
+	void RegisterTables(ITypedElementDataStorageInterface& DataStorage) const override;
+	void RegisterQueries(ITypedElementDataStorageInterface& DataStorage) const override;
+
+private:
+	void RegisterFetchUpdates(ITypedElementDataStorageInterface& DataStorage) const;
+	mutable TypedElementDataStorage::QueryHandle FetchUpdates = TypedElementDataStorage::InvalidQueryHandle;
+};
diff --git a/Engine/Plugins/Experimental/TypedElementsDataStorage/Source/TypedElementsDataStorageRevisionControl/Source/Private/TypedElementsDataStorageRevisionControl.cpp b/Engine/Plugins/Experimental/TypedElementsDataStorage/Source/TypedElementsDataStorageRevisionControl/Source/Private/TypedElementsDataStorageRevisionControl.cpp
new file mode 100644
--- /dev/null
+++ b/Engine/Plugins/Experimental/TypedElementsDataStorage/Source/TypedElementsDataStorageRevisionControl/Source/Private/TypedElementsDataStorageRevisionControl.cpp
@@ -0,0 +1,22 @@
+// Copyright Epic Games, Inc. All Rights Reserved.
+
+#include "../Public/TypedElementsDataStorageRevisionControl.h"
+
+void FTypedElementsDataStorageRevisionControlModule::StartupModule()
+{
+}
+
+void FTypedElementsDataStorageRevisionControlModule::ShutdownModule()
+{
+}
+
+void FTypedElementsDataStorageRevisionControlModule::AddReferencedObjects(FReferenceCollector& Collector)
+{
+}
+
+FString FTypedElementsDataStorageRevisionControlModule::GetReferencerName() const
+{
+	return TEXT("Typed Elements: Revision Control Module");
+}
+
+IMPLEMENT_MODULE(FTypedElementsDataStorageRevisionControlModule, TypedElementsDataStorageRevisionControl)
diff --git a/Engine/Plugins/Experimental/TypedElementsDataStorage/Source/TypedElementsDataStorageRevisionControl/Source/Public/TypedElementsDataStorageRevisionControl.h b/Engine/Plugins/Experimental/TypedElementsDataStorage/Source/TypedElementsDataStorageRevisionControl/Source/Public/TypedElementsDataStorageRevisionControl.h
new file mode 100644
--- /dev/null
+++ b/Engine/Plugins/Experimental/TypedElementsDataStorage/Source/TypedElementsDataStorageRevisionControl/Source/Public/TypedElementsDataStorageRevisionControl.h
@@ -0,0 +1,32 @@
+// Copyright Epic Games, Inc. All Rights Reserved.
+
+#pragma once
+
+#include "CoreMinimal.h"
+#include "Modules/ModuleManager.h"
+#include "UObject/GCObject.h"
+#include "UObject/ObjectPtr.h"
+
+class FReferenceCollector;
+class SWindow;
+
+class FTypedElementsDataStorageRevisionControlModule : public IModuleInterface, public FGCObject
+{
+public:
+	~FTypedElementsDataStorageRevisionControlModule() override = default;
+
+	//
+	// IModuleInterface
+	//
+
+	virtual void StartupModule() override;
+	virtual void ShutdownModule() override;
+
+	
+	//
+	// FGCObject
+	//
+	
+	void AddReferencedObjects(FReferenceCollector& Collector) override;
+	FString GetReferencerName() const override;
+};
diff --git a/Engine/Plugins/Experimental/TypedElementsDataStorage/Source/TypedElementsDataStorageRevisionControl/TypedElementsDataStorageRevisionControl.Build.CS b/Engine/Plugins/Experimental/TypedElementsDataStorage/Source/TypedElementsDataStorageRevisionControl/TypedElementsDataStorageRevisionControl.Build.CS
new file mode 100644
--- /dev/null
+++ b/Engine/Plugins/Experimental/TypedElementsDataStorage/Source/TypedElementsDataStorageRevisionControl/TypedElementsDataStorageRevisionControl.Build.CS
@@ -0,0 +1,44 @@
+// Copyright Epic Games, Inc. All Rights Reserved.
+
+using UnrealBuildTool;
+
+public class TypedElementsDataStorageRevisionControl : ModuleRules
+{
+	public TypedElementsDataStorageRevisionControl(ReadOnlyTargetRules Target) : base(Target)
+	{
+		PCHUsage = ModuleRules.PCHUsageMode.UseExplicitOrSharedPCHs;
+
+		if (Target.bBuildEditor)
+		{
+			PublicIncludePaths.AddRange(new string[] {});
+			PrivateIncludePaths.AddRange(new string[] {});
+
+			PublicDependencyModuleNames.AddRange(
+				new string[]
+				{
+					"Core",
+					"CoreUObject",
+					"Engine",
+					"EditorFramework",
+					"Slate",
+					"SlateCore",
+					"SourceControl",
+					"TypedElementFramework",
+					"TypedElementsDataStorage",
+					"EditorSubsystem",
+					"UnrealEd",
+				});
+			PrivateDependencyModuleNames.AddRange(
+				new string[]
+				{
+					"MassActors",
+					"MassEntity",
+					"StructUtils"
+				});
+
+			DynamicallyLoadedModuleNames.AddRange(new string[] {});
+			
+			ShortName = "TElmntsDataStorRevCtrl";
+		}
+	}
+}
\ No newline at end of file
diff --git a/Engine/Plugins/Experimental/TypedElementsDataStorage/TypedElementsDataStorage.uplugin b/Engine/Plugins/Experimental/TypedElementsDataStorage/TypedElementsDataStorage.uplugin
--- a/Engine/Plugins/Experimental/TypedElementsDataStorage/TypedElementsDataStorage.uplugin
+++ b/Engine/Plugins/Experimental/TypedElementsDataStorage/TypedElementsDataStorage.uplugin
@@ -22,6 +22,10 @@
 		{
 			"Name": "TypedElementsDataStorageUI",
 			"Type": "Editor"
+		},
+		{
+			"Name": "TypedElementsDataStorageRevisionControl",
+			"Type": "Editor"
 		}
 	],
 	"Plugins": [
diff --git a/Engine/Source/Runtime/TypedElementFramework/Public/Elements/Columns/TypedElementPackageColumns.h b/Engine/Source/Runtime/TypedElementFramework/Public/Elements/Columns/TypedElementPackageColumns.h
--- a/Engine/Source/Runtime/TypedElementFramework/Public/Elements/Columns/TypedElementPackageColumns.h
+++ b/Engine/Source/Runtime/TypedElementFramework/Public/Elements/Columns/TypedElementPackageColumns.h
@@ -18,6 +18,8 @@ struct FTypedElementPackageUnresolvedReference final : public FTypedElementDataS
 {
 	GENERATED_BODY()
 
+	FString PathOnDisk;
+	
 	TypedElementDataStorage::IndexHash Index;
 };
 
