Fixed a race condition in TEDS' Scratch Buffer.

The call to `GetEmptyBlock()` no longer uses an owner index, but instead uses a more traditional implementation of lock-free queue. The original implementation tried to avoid the ABA problem but since later updates to the scratch buffer started using seperate queues for available and recycled blocks, the ABA problem it tried to solve can no longer happen. A simpler implementation is therefore warranted and also fixes the deadlock.
The following changes are made in addition:
- It's now possible to have multiple instances of the scratch buffer. This has mostly been added to supporte testing.
- The scratch buffer is now stored in a shared pointer. This was needed to allow multiple instances of the scratch buffer to exist. Some optimization opportunities remain.
- Tests were added to cover basic functionality and do multi-threaded stress testing.

diff --git a/Engine/Plugins/Experimental/TypedElementsDataStorage/Source/TypedElementsDataStorage/Private/TypedElementDatabaseEnvironment.cpp b/Engine/Plugins/Experimental/TypedElementsDataStorage/Source/TypedElementsDataStorage/Private/TypedElementDatabaseEnvironment.cpp
--- a/Engine/Plugins/Experimental/TypedElementsDataStorage/Source/TypedElementsDataStorage/Private/TypedElementDatabaseEnvironment.cpp
+++ b/Engine/Plugins/Experimental/TypedElementsDataStorage/Source/TypedElementsDataStorage/Private/TypedElementDatabaseEnvironment.cpp
@@ -4,10 +4,10 @@
 
 FTypedElementDatabaseEnvironment::FTypedElementDatabaseEnvironment(
 	FMassEntityManager& InMassEntityManager, FMassProcessingPhaseManager& InMassPhaseManager)
-	: MassEntityManager(InMassEntityManager)
+	: ScratchBuffer(MakeShared<FTypedElementDatabaseScratchBuffer>())
+	, MassEntityManager(InMassEntityManager)
 	, MassPhaseManager(InMassPhaseManager)
 {
-
 }
 
 FTypedElementDatabaseIndexTable& FTypedElementDatabaseEnvironment::GetIndexTable()
@@ -22,12 +22,12 @@ const FTypedElementDatabaseIndexTable& FTypedElementDatabaseEnvironment::GetInde
 
 FTypedElementDatabaseScratchBuffer& FTypedElementDatabaseEnvironment::GetScratchBuffer()
 {
-	return ScratchBuffer;
+	return *ScratchBuffer;
 }
 
 const FTypedElementDatabaseScratchBuffer& FTypedElementDatabaseEnvironment::GetScratchBuffer() const
 {
-	return ScratchBuffer;
+	return *ScratchBuffer;
 }
 
 FMassEntityManager& FTypedElementDatabaseEnvironment::GetMassEntityManager()
diff --git a/Engine/Plugins/Experimental/TypedElementsDataStorage/Source/TypedElementsDataStorage/Private/TypedElementDatabaseEnvironment.h b/Engine/Plugins/Experimental/TypedElementsDataStorage/Source/TypedElementsDataStorage/Private/TypedElementDatabaseEnvironment.h
--- a/Engine/Plugins/Experimental/TypedElementsDataStorage/Source/TypedElementsDataStorage/Private/TypedElementDatabaseEnvironment.h
+++ b/Engine/Plugins/Experimental/TypedElementsDataStorage/Source/TypedElementsDataStorage/Private/TypedElementDatabaseEnvironment.h
@@ -26,8 +26,8 @@ public:
 
 private:
 	FTypedElementDatabaseIndexTable IndexTable;
-	FTypedElementDatabaseScratchBuffer ScratchBuffer;
-
+	TSharedPtr<FTypedElementDatabaseScratchBuffer> ScratchBuffer;
+	
 	FMassEntityManager& MassEntityManager;
 	FMassProcessingPhaseManager& MassPhaseManager;
 };
diff --git a/Engine/Plugins/Experimental/TypedElementsDataStorage/Source/TypedElementsDataStorage/Private/TypedElementDatabaseScratchBuffer.cpp b/Engine/Plugins/Experimental/TypedElementsDataStorage/Source/TypedElementsDataStorage/Private/TypedElementDatabaseScratchBuffer.cpp
--- a/Engine/Plugins/Experimental/TypedElementsDataStorage/Source/TypedElementsDataStorage/Private/TypedElementDatabaseScratchBuffer.cpp
+++ b/Engine/Plugins/Experimental/TypedElementsDataStorage/Source/TypedElementsDataStorage/Private/TypedElementDatabaseScratchBuffer.cpp
@@ -2,21 +2,19 @@
 
 #include "TypedElementDatabaseScratchBuffer.h"
 
+// Set to 1 to run additional tests on the scratch buffer. These are normally to expensive to run, but can help find deeper issues with the
+// scratch buffer.
+#define TEDS_SCRATCHBUFFER_ENABLE_ADDITIONAL_TESTS 0
+
 FTypedElementDatabaseScratchBuffer::~FTypedElementDatabaseScratchBuffer()
 {
 	// This assumes it's being called after all threads using the scratch buffer have been shutdown already so the
 	// scratch buffers they hold have been released back into the pool.
 
+	NextFrame(); // Force to next frame so any lingering blocks are freed correctly.
 	RecycleBlocks();
-	FBlock* FrontFull = FullBlocks;
-	while (FrontFull)
-	{
-		FBlock* Next = FrontFull->NextBlock;
-		delete FrontFull;
-		FrontFull = Next;
-	}
-	FullBlocks = nullptr;
-
+	checkf(FullBlocks == nullptr, TEXT("Not all full blocks were recycled. This can lead to memory leaks."));
+	
 	FBlock* FrontAvailable = AvailableBlocks;
 	while (FrontAvailable)
 	{
@@ -52,11 +50,35 @@ void FTypedElementDatabaseScratchBuffer::RecycleBlocks()
 		// relatively simple to destroy. The dirty blocks are added in front of the available blocks instead of the end to avoid having
 		// to iterate over the entire list of available blocks.
 
+#if TEDS_SCRATCHBUFFER_ENABLE_ADDITIONAL_TESTS
+		// Check to see if there are any cycles.
+		{
+			TSet<FBlock*> CycleDetector;
+			FBlock* TestBlock = FullBlocks;
+			while (TestBlock)
+			{
+				checkf(CycleDetector.Find(TestBlock) == nullptr, TEXT("A circular dependency was found in scratch buffer."));
+				CycleDetector.Add(TestBlock);
+				TestBlock = TestBlock->NextBlock;
+			}
+		}
+		{
+			TSet<FBlock*> CycleDetector;
+			FBlock* TestBlock = AvailableBlocks;
+			while (TestBlock)
+			{
+				checkf(CycleDetector.Find(TestBlock) == nullptr, TEXT("A circular dependency was found in scratch buffer."));
+				CycleDetector.Add(TestBlock);
+				TestBlock = TestBlock->NextBlock;
+			}
+		}
+#endif // TEDS_SCRATCHBUFFER_ENABLE_ADDITIONAL_TESTS
+
 		FBlock* RemainingBlocks = nullptr;
 		FBlock* DirtyBlock = FullBlocks;
 		do
 		{
-			// If the block hasn't been touched for at least one frame it can be deleted.
+			// If the block hasn't been touched for at least one frame it can be recycled.
 			if (DirtyBlock->LastTouchedByFrame < FrameId)
 			{
 				// Call destructors on any memory that needs it.
@@ -71,8 +93,9 @@ void FTypedElementDatabaseScratchBuffer::RecycleBlocks()
 				// The block is now clean so set the next block to process.
 				FBlock* CleanBlock = DirtyBlock;
 				DirtyBlock = DirtyBlock->NextBlock;
-
+				
 				// Reinsert the clean block into the chain of available blocks.
+				checkf(AvailableBlocks != CleanBlock, TEXT("Recycled block has already been added to the available blocks."));
 				CleanBlock->NextBlock.store(AvailableBlocks);
 				AvailableBlocks = CleanBlock;
 			}
@@ -90,13 +113,43 @@ void FTypedElementDatabaseScratchBuffer::RecycleBlocks()
 		
 		// Set the full blocks to the first block that couldn't be removed or null if all filled up blocks have been recycled.
 		FullBlocks = RemainingBlocks;
+
+#if TEDS_SCRATCHBUFFER_ENABLE_ADDITIONAL_TESTS
+		// Check if no new cycles have been introduced after updating.
+		{
+			TSet<FBlock*> CycleDetector;
+			FBlock* TestBlock = FullBlocks;
+			while (TestBlock)
+			{
+				checkf(CycleDetector.Find(TestBlock) == nullptr, TEXT("A circular dependency was found in scratch buffer."));
+				checkf(TestBlock->LastTouchedByFrame >= FrameId, TEXT("Deleted block still referenced."));
+				CycleDetector.Add(TestBlock);
+				TestBlock = TestBlock->NextBlock;
+			}
+		}
+		{
+			TSet<FBlock*> CycleDetector;
+			FBlock* TestBlock = AvailableBlocks;
+			while (TestBlock)
+			{
+				checkf(CycleDetector.Find(TestBlock) == nullptr, TEXT("A circular dependency was found in scratch buffer."));
+				CycleDetector.Add(TestBlock);
+				TestBlock = TestBlock->NextBlock;
+			}
+		}
+#endif // TEDS_SCRATCHBUFFER_ENABLE_ADDITIONAL_TESTS
 	}
 }
 
+FTypedElementDatabaseScratchBuffer::FBlockControllerMap& FTypedElementDatabaseScratchBuffer::GetThreadLocalBlockControllerMap()
+{
+	thread_local static FBlockControllerMap LocalBlockMap;
+	return LocalBlockMap;
+}
+
 FTypedElementDatabaseScratchBuffer::FBlockController& FTypedElementDatabaseScratchBuffer::GetThreadLocalBlockController()
 {
-	thread_local static FBlockController LocalBlock = FBlockController(*this);
-	return LocalBlock;
+	return GetThreadLocalBlockControllerMap().FindOrAddControllerFor(*this);
 }
 
 void FTypedElementDatabaseScratchBuffer::ConfigureDestructorTail(
@@ -110,15 +163,15 @@ void FTypedElementDatabaseScratchBuffer::ConfigureDestructorTail(
 // FBlockController
 //
 
-FTypedElementDatabaseScratchBuffer::FBlockController::FBlockController(FTypedElementDatabaseScratchBuffer& InOwner)
-	: Owner(InOwner)
-	, Id(InOwner.BlockControllerId++)
+FTypedElementDatabaseScratchBuffer::FBlockController::FBlockController(const TWeakPtr<FTypedElementDatabaseScratchBuffer>& InParent)
+	: Parent(InParent)
 {
 	Block = GetEmptyBlock();
 }
 
 FTypedElementDatabaseScratchBuffer::FBlockController::~FBlockController()
 {
+	checkf(Block != nullptr, TEXT("FBlockController is expected to always have a valid block to work with."));
 	RecycleBlock();
 }
 
@@ -173,43 +226,45 @@ void* FTypedElementDatabaseScratchBuffer::FBlockController::Allocate(size_t Size
 
 FTypedElementDatabaseScratchBuffer::FBlock* FTypedElementDatabaseScratchBuffer::FBlockController::GetEmptyBlock()
 {
-	FBlock* FrontBlock = Owner.AvailableBlocks;
-	while (FrontBlock != nullptr)
+	if (TSharedPtr<FTypedElementDatabaseScratchBuffer> ParentInstance = Parent.Pin())
 	{
-		// If this is not zero it means another thread has already claimed this block so try again with the next block.
-		uint32 CurrentOwner = 0;
-		if (FrontBlock->Owner.compare_exchange_strong(CurrentOwner, Id))
+		FBlock* FrontBlock = ParentInstance->AvailableBlocks.load();
+		// Because the full and available queues are separate, this does not suffer from the ABA problem.
+		while (FrontBlock != nullptr && !ParentInstance->AvailableBlocks.compare_exchange_weak(FrontBlock, FrontBlock->NextBlock)){};
+		if (FrontBlock)
 		{
-			Owner.AvailableBlocks.store(FrontBlock->NextBlock.load());
 			FrontBlock->NextBlock = nullptr;
 			return FrontBlock;
 		}
-		else
-		{
-			// Get the new front block. It can happen due to the thread timing that the same block is retrieved, but 
-			// avoiding this adds a lot more complexity and contention on the front block is low.
-			FrontBlock = Owner.AvailableBlocks;
-		}
+	}
 
-	};
-	
 	// If null then there are no more buffers left so create a new one.
-	FBlock* NewBlock = new FBlock();
-	NewBlock->Owner = Id;
-	return NewBlock;
+	return new FBlock();
 }
 
 void FTypedElementDatabaseScratchBuffer::FBlockController::RecycleBlock()
 {
-	Block->Front = 0;
-	Block->Owner = 0;
-
-	// Replace the front of the linked list with full blocks with the latest full block.
-	FBlock* FrontBlock = Owner.FullBlocks;
-	do
+	if (TSharedPtr<FTypedElementDatabaseScratchBuffer> ParentInstance = Parent.Pin())
 	{
-		Block->NextBlock = FrontBlock;
-	} while (!Owner.FullBlocks.compare_exchange_strong(FrontBlock, Block));
+		Block->Front = 0;
+		
+		// Replace the front of the linked list with full blocks with the latest full block.
+		// Because the full and available queues are separate, this does not suffer from the ABA problem.
+		FBlock* FrontBlock = ParentInstance->FullBlocks;
+		do
+		{
+			Block->NextBlock = FrontBlock;
+		} while (!ParentInstance->FullBlocks.compare_exchange_strong(FrontBlock, Block));
+	}
+	else
+	{
+		// The assumption is that if the scratch buffer has been deleted none of the data in the
+		// current buffer will be used so it can be deleted. It also means that the scratch buffer
+		// is not going to recycle this block for this thread so needs to be deleted here to avoid
+		// a memory leak.
+		delete Block;
+	}
+	Block = nullptr;
 }
 
 void FTypedElementDatabaseScratchBuffer::FBlockController::ConfigureDestructorTail(
@@ -221,3 +276,27 @@ void FTypedElementDatabaseScratchBuffer::FBlockController::ConfigureDestructorTa
 	Destructor.InstanceCount = Count;
 	Block->DestructionTail = &Destructor;
 }
+
+
+
+//
+// FBlockControllerMap
+//
+
+FTypedElementDatabaseScratchBuffer::FBlockController& 
+	FTypedElementDatabaseScratchBuffer::FBlockControllerMap::FindOrAddControllerFor(FTypedElementDatabaseScratchBuffer& ScratchBuffer)
+{
+	TWeakPtr<FTypedElementDatabaseScratchBuffer>* ParentsIt = Parents.GetData();
+	const uint32 Count = Parents.Num();
+	for (uint32 Index = 0; Index < Count; ++Index, ++ParentsIt)
+	{
+		if (ParentsIt->HasSameObject(&ScratchBuffer))
+		{
+			return Controllers[Index];
+		}
+	}
+
+	Parents.Add(ScratchBuffer.AsWeak());
+	uint32 NewIndex = Controllers.Emplace(ScratchBuffer.AsWeak());
+	return Controllers[NewIndex];
+}
diff --git a/Engine/Plugins/Experimental/TypedElementsDataStorage/Source/TypedElementsDataStorage/Private/TypedElementDatabaseScratchBuffer.h b/Engine/Plugins/Experimental/TypedElementsDataStorage/Source/TypedElementsDataStorage/Private/TypedElementDatabaseScratchBuffer.h
--- a/Engine/Plugins/Experimental/TypedElementsDataStorage/Source/TypedElementsDataStorage/Private/TypedElementDatabaseScratchBuffer.h
+++ b/Engine/Plugins/Experimental/TypedElementsDataStorage/Source/TypedElementsDataStorage/Private/TypedElementDatabaseScratchBuffer.h
@@ -8,6 +8,7 @@
 #include "Math/UnrealMath.h"
 #include "Misc/MTAccessDetector.h"
 #include "Templates/AlignmentTemplates.h"
+#include "Templates/SharedPointer.h"
 #include "Templates/UnrealTemplate.h"
 
 // A thread-safe memory allocator that uses linear allocation using a chain of recycled memory blocks.
@@ -15,9 +16,10 @@
 // live at best until the end of the frame. Objects that require destruction will be deleted once a block is
 // recycled. A block is only recycled when it's full, so during low activity on a particular thread it can
 // take significantly longer than a single frame before an object is destroyed.
-class FTypedElementDatabaseScratchBuffer
+class FTypedElementDatabaseScratchBuffer : public TSharedFromThis<FTypedElementDatabaseScratchBuffer>
 {
 	friend class FBlockController;
+	friend class FBlockControllerMap;
 public:
 	~FTypedElementDatabaseScratchBuffer();
 
@@ -61,7 +63,6 @@ private:
 		uint64 LastTouchedByFrame = 0;
 		FDestructorTail* DestructionTail = nullptr;
 		std::atomic<FBlock*> NextBlock = nullptr;
-		std::atomic<uint32> Owner = 0;
 		// Offset into the buffer where the next allocation starts.
 		uint32 Front = 0;
 	};
@@ -69,7 +70,7 @@ private:
 	// Simple wrapper around the buffer. This is to guarantee that when the thread gets destroyed the buffer is cleaned up.
 	struct FBlockController
 	{
-		explicit FBlockController(FTypedElementDatabaseScratchBuffer& InOwner);
+		explicit FBlockController(const TWeakPtr<FTypedElementDatabaseScratchBuffer>& InParent);
 		~FBlockController();
 
 		void* Allocate(size_t Size, size_t Alignment, uint64 LocalFrameId);
@@ -77,11 +78,21 @@ private:
 		void RecycleBlock();
 		void ConfigureDestructorTail(FDestructorTail& Destructor, DestructorFunction Callback, void* Object, int32 Count = 1);
 
-		FTypedElementDatabaseScratchBuffer& Owner;
+		TWeakPtr<FTypedElementDatabaseScratchBuffer> Parent;
 		FBlock* Block;
-		uint32 Id;
 	};
 
+	// Simple container to hold multiple block controllers and return the one appropriate for the thread and scratch buffer.
+	struct FBlockControllerMap
+	{
+		TArray<TWeakPtr<FTypedElementDatabaseScratchBuffer>> Parents;
+		TArray<FBlockController> Controllers;
+
+		FBlockController& FindOrAddControllerFor(FTypedElementDatabaseScratchBuffer& ScratchBuffer);
+	};
+
+	// Returns the block controller map unique to the thread that calls this function.
+	FBlockControllerMap& GetThreadLocalBlockControllerMap();
 	// Returns the block controller that's unique to the thread that calls this function.
 	FBlockController& GetThreadLocalBlockController();
 	void ConfigureDestructorTail(FDestructorTail& Destructor, DestructorFunction Callback, void* Object, int32 Count = 1);
@@ -91,10 +102,7 @@ private:
 
 	std::atomic<FBlock*> AvailableBlocks = nullptr;
 	std::atomic<FBlock*> FullBlocks = nullptr;
-	// Running counter so each FBlockController gets a unique id to identify the blocks assigned to them. Zero is reserved 
-	// to indicate it's not in use by a block controller.
-	std::atomic<uint32> BlockControllerId = 1;
-
+	
 	// The number that identifies the current frame. Used to keep track of the blocks that were used in this frame or earlier.
 	uint64 FrameId = 0;
 };
diff --git a/Engine/Plugins/Experimental/TypedElementsDataStorage/Source/TypedElementsDataStorage/Tests/TypedElementDatabaseScratchBufferTests.cpp b/Engine/Plugins/Experimental/TypedElementsDataStorage/Source/TypedElementsDataStorage/Tests/TypedElementDatabaseScratchBufferTests.cpp
new file mode 100644
--- /dev/null
+++ b/Engine/Plugins/Experimental/TypedElementsDataStorage/Source/TypedElementsDataStorage/Tests/TypedElementDatabaseScratchBufferTests.cpp
@@ -0,0 +1,253 @@
+// Copyright Epic Games, Inc. All Rights Reserved.
+
+#if WITH_TESTS
+
+#include <chrono>
+#include <condition_variable>
+#include <thread>
+#include "TypedElementDatabaseScratchBuffer.h"
+#include "Tests/TestHarnessAdapter.h"
+
+// Tries to shutdown all threads within the maximum wait time. If it takes longer than the minimum wait time the test will fail.
+template<int ThreadCount, class Rep, class Period>
+static void WaitForFinalization(std::condition_variable& WaitVariable, std::mutex& WaitMutex, 
+	std::atomic<bool>& bRunFlag, std::atomic<uint32>& CompletedThreadCount, const std::chrono::duration<Rep, Period>& MaxWaitTime)
+{
+	std::unique_lock<std::mutex> WaitLock(WaitMutex); // Wait to be re-enabled again.
+	bRunFlag = false;
+	std::chrono::system_clock::time_point StartTime = std::chrono::system_clock::now();
+	while (CompletedThreadCount != ThreadCount 
+		&& std::chrono::duration_cast<std::chrono::milliseconds>(std::chrono::system_clock::now() - StartTime) < MaxWaitTime)
+	{
+		WaitVariable.wait_for(WaitLock, MaxWaitTime);
+	}
+	REQUIRE(CompletedThreadCount == ThreadCount);
+}
+
+// Spin locks for the maximum wait time to let threads resume
+template<int ThreadCount, class Rep, class Period>
+static void WaitForPause(std::atomic<uint32>& PausedCounter, const std::chrono::duration<Rep, Period>& MaxWaitTime)
+{
+	std::chrono::system_clock::time_point StartTime = std::chrono::system_clock::now();
+	while (PausedCounter != ThreadCount
+		&& std::chrono::duration_cast<std::chrono::milliseconds>(std::chrono::system_clock::now() - StartTime) < MaxWaitTime)
+	{}
+	REQUIRE(PausedCounter == ThreadCount);
+}
+
+TEST_CASE("TypedElementsDataStorage::Scratch Buffer - MT (FTypedElementDatabaseScratchBuffer)", "[ApplicationContextMask][EngineFilter]")
+{
+	SECTION("Stress test without recycling.")
+	{
+		using namespace std::chrono_literals;
+		constexpr uint32 ThreadCount = 8;
+
+		TSharedPtr<FTypedElementDatabaseScratchBuffer> Buffer = MakeShared<FTypedElementDatabaseScratchBuffer>();
+
+		std::atomic<bool> bKeepRunning = true;
+		std::atomic<uint32> FailedAllocations = 0;
+		std::atomic<uint32> CompletedThreadCount = 0;
+		std::condition_variable CompletionSync;
+		std::mutex CompletionMutex;
+		TArray<std::thread> Threads;
+		Threads.Reserve(ThreadCount);
+		for (uint32 ThreadCounter = 0; ThreadCounter < ThreadCount; ++ThreadCounter)
+		{
+			Threads.Emplace([&Buffer, &bKeepRunning, &FailedAllocations, &CompletedThreadCount, &CompletionSync]()
+				{
+					while (bKeepRunning)
+					{
+						void* Data = Buffer->Allocate(128, 4);
+						if (Data == nullptr)
+						{
+							FailedAllocations++;
+						}
+					}
+					if (++CompletedThreadCount == ThreadCount)
+					{
+						CompletionSync.notify_all();
+					}
+				});
+		}
+
+		std::this_thread::sleep_for(2000ms);
+		WaitForFinalization<ThreadCount>(CompletionSync, CompletionMutex, bKeepRunning, CompletedThreadCount, 500ms);
+		
+		for (std::thread& Thread : Threads)
+		{
+			Thread.join();
+		}
+
+		CHECK(FailedAllocations == 0);
+	}
+
+	SECTION("Stress test with recycling.")
+	{
+		using namespace std::chrono_literals;
+
+		constexpr uint32 ThreadCount = 8;
+
+		TSharedPtr<FTypedElementDatabaseScratchBuffer> Buffer = MakeShared<FTypedElementDatabaseScratchBuffer>();
+
+		std::atomic<bool> bKeepRunning = true;
+		std::atomic<bool> bPause = false;
+		std::atomic<uint32> PausedCounter = 0;
+		std::atomic<uint32> FailedAllocations = 0;
+
+		std::condition_variable PauseSync;
+		std::mutex PauseMutex;
+
+		std::atomic<uint32> CompletedThreadCount = 0;
+		std::condition_variable CompletionSync;
+		std::mutex CompletionMutex;
+
+		TArray<std::thread> Threads;
+		Threads.Reserve(ThreadCount);
+		for (uint32 ThreadCounter = 0; ThreadCounter < ThreadCount; ++ThreadCounter)
+		{
+			Threads.Emplace([
+				ThreadCounter, &Buffer, &bKeepRunning, &FailedAllocations, &bPause, &PausedCounter, 
+				&PauseSync, &PauseMutex, &CompletedThreadCount, &CompletionSync]()
+				{
+					static constexpr uint32 MemoryAllocationSizes[] = { 128, 32, 14332, 741, 8871, 48, 27335 };
+					static constexpr uint32 MemoryAllocationSizeCount = sizeof(MemoryAllocationSizes) / sizeof(MemoryAllocationSizes[0]);
+					
+					uint32 MemoryAllocationSizeIndex = ThreadCounter;
+					while (bKeepRunning)
+					{
+						// Not using % as there's currently a false positive in MSVC that marks it warning (C6385).
+						++MemoryAllocationSizeIndex;
+						if (MemoryAllocationSizeIndex >= MemoryAllocationSizeCount)
+						{
+							MemoryAllocationSizeIndex = 0;
+						}
+
+						void* Data = Buffer->Allocate(static_cast<size_t>(MemoryAllocationSizes[MemoryAllocationSizeIndex]), 4);
+						if (Data == nullptr)
+						{
+							FailedAllocations++;
+						}
+
+						if (bPause)
+						{
+							++PausedCounter; // Announce that this thread is paused.
+							std::unique_lock<std::mutex> WaitLock(PauseMutex); // Wait to be re-enabled again.
+							PauseSync.wait(WaitLock);
+							++PausedCounter; // Announce that this thread has resumed.
+						}
+					}
+					if (++CompletedThreadCount == ThreadCount)
+					{
+						CompletionSync.notify_all();
+					}
+				});
+		}
+
+		// While the threads are continuously allocating memory, use the main thread to periodically pause the other threads 
+		// and recycle blocks.
+		for (uint32 IterationCounter = 0; IterationCounter < 60; ++IterationCounter)
+		{
+			PausedCounter = 0;
+			bPause = true;
+			WaitForPause<ThreadCount>(PausedCounter, 500ms);
+			
+			Buffer->NextFrame();
+			Buffer->RecycleBlocks();
+
+			// Resume all threads.
+			{
+				std::unique_lock<std::mutex> WaitLock(PauseMutex);
+				PausedCounter = 0;
+				bPause = false;
+				PauseSync.notify_all();
+			}
+			WaitForPause<ThreadCount>(PausedCounter, 500ms);
+		
+			std::this_thread::sleep_for(33ms);
+		}
+		WaitForFinalization<ThreadCount>(CompletionSync, CompletionMutex, bKeepRunning, CompletedThreadCount, 500ms);
+
+		for (std::thread& Thread : Threads)
+		{
+			Thread.join();
+		}
+
+		CHECK(FailedAllocations == 0);
+	}
+}
+
+TEST_CASE("TypedElementsDataStorage::Scratch Buffer (FTypedElementDatabaseScratchBuffer)", "[ApplicationContextMask][EngineFilter]")
+{
+	SECTION("Create and destroy buffer")
+	{
+		FTypedElementDatabaseScratchBuffer Buffer;
+	}
+
+	SECTION("Allocate small block")
+	{
+		TSharedPtr<FTypedElementDatabaseScratchBuffer> Buffer = MakeShared<FTypedElementDatabaseScratchBuffer>();
+		void* Data = Buffer->Allocate(4, 4);
+		CHECK(Data != nullptr);
+	}
+
+	SECTION("Allocate over-sized block")
+	{
+		TSharedPtr<FTypedElementDatabaseScratchBuffer> Buffer = MakeShared<FTypedElementDatabaseScratchBuffer>();
+		void* Data = Buffer->Allocate(Buffer->MaxAllocationSize() * 4, 4);
+		CHECK(Data != nullptr);
+	}
+
+	SECTION("Allocate exactly remaining block size")
+	{
+		TSharedPtr<FTypedElementDatabaseScratchBuffer> Buffer = MakeShared<FTypedElementDatabaseScratchBuffer>();
+		void* Data1 = Buffer->Allocate(1, 1);
+		void* Data2 = Buffer->Allocate(Buffer->MaxAllocationSize() - 1, 1);
+
+		REQUIRE(Data1 != nullptr);
+		REQUIRE(Data2 != nullptr);
+		// Check if they're both in the same buffer.
+		CHECK(reinterpret_cast<uintptr_t>(Data2) == reinterpret_cast<uintptr_t>(Data1) + 1);
+	}
+
+	SECTION("Alignment respected")
+	{
+		TSharedPtr<FTypedElementDatabaseScratchBuffer> Buffer = MakeShared<FTypedElementDatabaseScratchBuffer>();
+		void* Data1 = Buffer->Allocate(1, 1);
+		void* Data2 = Buffer->Allocate(4, 4);
+
+		REQUIRE(Data1 != nullptr);
+		REQUIRE(Data2 != nullptr);
+		// Check if they're both in the same buffer.
+		CHECK((reinterpret_cast<uintptr_t>(Data2) & 3) == 0);
+	}
+
+	SECTION("Multiple blocks used.")
+	{
+		TSharedPtr<FTypedElementDatabaseScratchBuffer> Buffer = MakeShared<FTypedElementDatabaseScratchBuffer>();
+
+		int32 IncrementCount = (Buffer->MaxAllocationSize() * 4 /* Fill 4 blocks*/) / 64 /* With 64 by allocations */;
+		for (int32 Counter = 0; Counter < IncrementCount; ++Counter)
+		{
+			void* Data = Buffer->Allocate(64, 4);
+			REQUIRE(Data != nullptr);
+		}
+	}
+
+	SECTION("Recycle full blocks.")
+	{
+		TSharedPtr<FTypedElementDatabaseScratchBuffer> Buffer = MakeShared<FTypedElementDatabaseScratchBuffer>();
+
+		for (int32 Iterations = 0; Iterations < 16; ++Iterations)
+		{
+			int32 IncrementCount = Buffer->MaxAllocationSize() / 64;
+			for (int32 Counter = 0; Counter <= IncrementCount; ++Counter) // Increment one more that needed to fill the block and make sure it's recycled.
+			{
+				void* Data = Buffer->Allocate(64, 4);
+				REQUIRE(Data != nullptr);
+			}
+			Buffer->NextFrame();
+			Buffer->RecycleBlocks();
+		}
+	}
+}
+#endif // WITH_TESTS
