Only sync UObject package information to TEDS when revision control syncing is also enabled

Prevents some instability in the short term

diff --git a/Engine/Plugins/Experimental/TypedElementsDataStorage/Source/TypedElementsDataStorage/Private/Compatibility/TypedElementUObjectPackagePathToColumnQueries.cpp b/Engine/Plugins/Experimental/TypedElementsDataStorage/Source/TypedElementsDataStorage/Private/Compatibility/TypedElementUObjectPackagePathToColumnQueries.cpp
--- a/Engine/Plugins/Experimental/TypedElementsDataStorage/Source/TypedElementsDataStorage/Private/Compatibility/TypedElementUObjectPackagePathToColumnQueries.cpp
+++ b/Engine/Plugins/Experimental/TypedElementsDataStorage/Source/TypedElementsDataStorage/Private/Compatibility/TypedElementUObjectPackagePathToColumnQueries.cpp
@@ -10,6 +10,13 @@
 #include "Misc/Paths.h"
 #include "UObject/Package.h"
 
+static bool bAutoPopulateRevisionControlState = false;
+TYPEDELEMENTSDATASTORAGE_API FAutoConsoleVariableRef CVarAutoPopulateState(
+	TEXT("TEDS.RevisionControl.AutoPopulateState"),
+	bAutoPopulateRevisionControlState,
+	TEXT("Automatically query revision control provider and fill information into TEDS")
+);
+
 static void ResolvePackageReference(ITypedElementDataStorageInterface::IQueryContext& Context, const UPackage* Package, TypedElementRowHandle Row, TypedElementRowHandle PackageRow)
 {
 	FTypedElementPackageReference PackageReference;
@@ -30,8 +37,58 @@ void UTypedElementUObjectPackagePathFactory::RegisterQueries(ITypedElementDataSt
 {
 	using namespace TypedElementQueryBuilder;
 	using DSI = ITypedElementDataStorageInterface;
-	
+
+	CVarAutoPopulateState->AsVariable()->OnChangedDelegate().AddLambda(
+		[this, &DataStorage](IConsoleVariable* AutoPopulate)
+		{
+			if (AutoPopulate->GetBool())
+			{
+				RegisterTryAddPackageRef(DataStorage);
+			}
+			else
+			{
+				DataStorage.UnregisterQuery(TryAddPackageRef);
+			}
+		}
+	);
+
 	DataStorage.RegisterQuery(
+		Select(
+			TEXT("Resolve package references"),
+			FProcessor(DSI::EQueryTickPhase::FrameEnd, DataStorage.GetQueryTickGroupName(DSI::EQueryTickGroups::SyncExternalToDataStorage)),
+			[](DSI::IQueryContext& Context, TypedElementRowHandle Row, const FTypedElementUObjectColumn& Object, const FTypedElementPackageUnresolvedReference& UnresolvedPackageReference)
+			{
+				TypedElementRowHandle PackageRow = Context.FindIndexedRow(UnresolvedPackageReference.Index);
+				if (!Context.IsRowAvailable(PackageRow))
+				{
+					return;
+				}
+				const UObject* ObjectInstance = Object.Object.Get();
+				if (!ObjectInstance)
+				{
+					return;
+				}
+				const UPackage* Package = ObjectInstance->GetPackage();
+				Context.RemoveColumns(Row, { FTypedElementPackageUnresolvedReference::StaticStruct() });
+
+				ResolvePackageReference(Context, Package, Row, PackageRow);
+			}
+		)
+		.Compile()
+	);
+
+	if (CVarAutoPopulateState->GetBool())
+	{
+		RegisterTryAddPackageRef(DataStorage);
+	}
+}
+
+void UTypedElementUObjectPackagePathFactory::RegisterTryAddPackageRef(ITypedElementDataStorageInterface& DataStorage) const
+{
+	using namespace TypedElementQueryBuilder;
+	using DSI = ITypedElementDataStorageInterface;
+	
+	TryAddPackageRef = DataStorage.RegisterQuery(
 		Select(
 			TEXT("Sync UObject package info to columns"),
 			FObserver::OnAdd<FTypedElementUObjectColumn>()
@@ -57,9 +114,9 @@ void UTypedElementUObjectPackagePathFactory::RegisterQueries(ITypedElementDataSt
 						else
 						{
 							FTypedElementPackageUnresolvedReference UnresolvedPackageReference;
-                            UnresolvedPackageReference.Index = Index;
+							UnresolvedPackageReference.Index = Index;
 							UnresolvedPackageReference.PathOnDisk = MoveTemp(FullPackageFilename);
-                            Context.AddColumn(Row, MoveTemp(UnresolvedPackageReference));
+							Context.AddColumn(Row, MoveTemp(UnresolvedPackageReference));
 						}
 					}
 				}
@@ -67,29 +124,4 @@ void UTypedElementUObjectPackagePathFactory::RegisterQueries(ITypedElementDataSt
 		)
 		.Compile()
 	);
-
-	DataStorage.RegisterQuery(
-		Select(
-			TEXT("Resolve package references"),
-			FProcessor(DSI::EQueryTickPhase::FrameEnd, DataStorage.GetQueryTickGroupName(DSI::EQueryTickGroups::SyncExternalToDataStorage)),
-			[](DSI::IQueryContext& Context, TypedElementRowHandle Row, const FTypedElementUObjectColumn& Object, const FTypedElementPackageUnresolvedReference& UnresolvedPackageReference)
-			{
-				TypedElementRowHandle PackageRow = Context.FindIndexedRow(UnresolvedPackageReference.Index);
-				if (!Context.IsRowAvailable(PackageRow))
-				{
-					return;
-				}
-				const UObject* ObjectInstance = Object.Object.Get();
-				if (!ObjectInstance)
-				{
-					return;
-				}
-				const UPackage* Package = ObjectInstance->GetPackage();
-				Context.RemoveColumns(Row, { FTypedElementPackageUnresolvedReference::StaticStruct() });
-
-				ResolvePackageReference(Context, Package, Row, PackageRow);
-			}
-		)
-		.Compile()
-	);
 }
diff --git a/Engine/Plugins/Experimental/TypedElementsDataStorage/Source/TypedElementsDataStorage/Private/Compatibility/TypedElementUObjectPackagePathToColumnQueries.h b/Engine/Plugins/Experimental/TypedElementsDataStorage/Source/TypedElementsDataStorage/Private/Compatibility/TypedElementUObjectPackagePathToColumnQueries.h
--- a/Engine/Plugins/Experimental/TypedElementsDataStorage/Source/TypedElementsDataStorage/Private/Compatibility/TypedElementUObjectPackagePathToColumnQueries.h
+++ b/Engine/Plugins/Experimental/TypedElementsDataStorage/Source/TypedElementsDataStorage/Private/Compatibility/TypedElementUObjectPackagePathToColumnQueries.h
@@ -2,6 +2,7 @@
 
 #pragma once
 
+#include "Elements/Common/TypedElementHandles.h"
 #include "Elements/Interfaces/TypedElementDataStorageFactory.h"
 #include "UObject/ObjectMacros.h"
 
@@ -16,4 +17,8 @@ public:
 	~UTypedElementUObjectPackagePathFactory() override = default;
 
 	void RegisterQueries(ITypedElementDataStorageInterface& DataStorage) const override;
+
+private:
+	void RegisterTryAddPackageRef(ITypedElementDataStorageInterface& DataStorage) const;
+	mutable TypedElementDataStorage::QueryHandle TryAddPackageRef = TypedElementDataStorage::InvalidQueryHandle;
 };
diff --git a/Engine/Plugins/Experimental/TypedElementsDataStorage/Source/TypedElementsDataStorageRevisionControl/Source/Private/RevisionControlProcessors.cpp b/Engine/Plugins/Experimental/TypedElementsDataStorage/Source/TypedElementsDataStorageRevisionControl/Source/Private/RevisionControlProcessors.cpp
--- a/Engine/Plugins/Experimental/TypedElementsDataStorage/Source/TypedElementsDataStorageRevisionControl/Source/Private/RevisionControlProcessors.cpp
+++ b/Engine/Plugins/Experimental/TypedElementsDataStorage/Source/TypedElementsDataStorageRevisionControl/Source/Private/RevisionControlProcessors.cpp
@@ -14,12 +14,7 @@
 #include "Elements/Framework/TypedElementQueryBuilder.h"
 #include "Elements/Framework/TypedElementRegistry.h"
 
-static bool bAutoPopulateRevisionControlState = false;
-static FAutoConsoleVariableRef CVarAutoPopulateState(
-	TEXT("TEDS.RevisionControl.AutoPopulateState"),
-	bAutoPopulateRevisionControlState,
-	TEXT("Automatically query revision control provider and fill information into TEDS")
-);
+extern TYPEDELEMENTSDATASTORAGE_API FAutoConsoleVariableRef CVarAutoPopulateState;
 
 FAutoConsoleCommandWithArgsAndOutputDevice SetSelectionSCCStateConsoleCommand(
 	TEXT("TEDS.Debug.SetSCCState"),
@@ -111,9 +106,9 @@ void UTypedElementRevisionControlFactory::RegisterQueries(ITypedElementDataStora
 			.All<FTypedElementSyncFromWorldTag>()
 		.Compile()
 	);
-
-	CVarAutoPopulateState->AsVariable()->SetOnChangedCallback(
-		FConsoleVariableDelegate::CreateLambda([this, &DataStorage](IConsoleVariable* AutoPopulate)
+	
+	CVarAutoPopulateState->AsVariable()->OnChangedDelegate().AddLambda(
+		[this, &DataStorage](IConsoleVariable* AutoPopulate)
 		{
 			if (AutoPopulate->GetBool())
 			{
@@ -123,10 +118,10 @@ void UTypedElementRevisionControlFactory::RegisterQueries(ITypedElementDataStora
 			{
 				DataStorage.UnregisterQuery(FetchUpdates);
 			}
-		})
+		}
 	); 
 	
-	if (bAutoPopulateRevisionControlState)
+	if (CVarAutoPopulateState->GetBool())
 	{
 		RegisterFetchUpdates(DataStorage);
 	}
