Fix bug in UTypedElementDatabase::BatchAddRemoveColumns leading to incorrect row compositions

* Existing implementation incorrectly computed bMustUpdateFragments and bMustUpdateTags.
* Fix was to set bMustUpdateFragments based on presence of a struct in either Add or Remove arrays that derive from a FMassFragment.  Similar logic for bMustUpdateTags

diff --git a/Engine/Plugins/Experimental/TypedElementsDataStorage/Source/TypedElementsDataStorage/Private/TypedElementDatabase.cpp b/Engine/Plugins/Experimental/TypedElementsDataStorage/Source/TypedElementsDataStorage/Private/TypedElementDatabase.cpp
--- a/Engine/Plugins/Experimental/TypedElementsDataStorage/Source/TypedElementsDataStorage/Private/TypedElementDatabase.cpp
+++ b/Engine/Plugins/Experimental/TypedElementsDataStorage/Source/TypedElementsDataStorage/Private/TypedElementDatabase.cpp
@@ -64,6 +64,36 @@ FAutoConsoleCommandWithOutputDevice PrintSupportedColumnsConsoleCommand(
 			Output.Log(TEXT("End of Typed Elements Data Storage supported column list."));
 		}));
 
+namespace TypedElementDatabasePrivate
+{
+	struct ColumnsToBitSetsResult
+	{
+		bool bMustUpdateFragments = false;
+		bool bMustUpdateTags = false;
+		
+		bool MustUpdate() const { return bMustUpdateFragments || bMustUpdateTags; }
+	};
+	ColumnsToBitSetsResult ColumnsToBitSets(TConstArrayView<const UScriptStruct*> Columns, FMassFragmentBitSet& Fragments, FMassTagBitSet& Tags)
+	{
+		ColumnsToBitSetsResult Result;
+
+		for (const UScriptStruct* ColumnType : Columns)
+		{
+			if (ColumnType->IsChildOf(FMassFragment::StaticStruct()))
+			{
+				Fragments.Add(*ColumnType);
+				Result.bMustUpdateFragments = true;
+			}
+			else if (ColumnType->IsChildOf(FMassTag::StaticStruct()))
+			{
+				Tags.Add(*ColumnType);
+				Result.bMustUpdateTags = true;
+			}
+		}
+		return Result;
+	}
+}
+
 void UTypedElementDatabase::Initialize()
 {
 	check(GEditor);
@@ -430,7 +460,7 @@ bool UTypedElementDatabase::AddColumns(TypedElementRowHandle Row, TConstArrayVie
 
 		FMassFragmentBitSet FragmentsToAdd;
 		FMassTagBitSet TagsToAdd;
-		if (ColumnsToBitSets(Columns, FragmentsToAdd, TagsToAdd))
+		if (TypedElementDatabasePrivate::ColumnsToBitSets(Columns, FragmentsToAdd, TagsToAdd).MustUpdate())
 		{
 			if (ActiveEditorEntityManager->IsEntityActive(Entity))
 			{
@@ -455,7 +485,7 @@ void UTypedElementDatabase::RemoveColumns(TypedElementRowHandle Row, TConstArray
 
 		FMassFragmentBitSet FragmentsToRemove;
 		FMassTagBitSet TagsToRemove;
-		if (ColumnsToBitSets(Columns, FragmentsToRemove, TagsToRemove))
+		if (TypedElementDatabasePrivate::ColumnsToBitSets(Columns, FragmentsToRemove, TagsToRemove).MustUpdate())
 		{
 			if (ActiveEditorEntityManager->IsEntityActive(Entity))
 			{
@@ -480,7 +510,7 @@ bool UTypedElementDatabase::AddRemoveColumns(TypedElementRowHandle Row,
 
 		FMassFragmentBitSet FragmentsToAdd;
 		FMassTagBitSet TagsToAdd;
-		if (ColumnsToBitSets(ColumnsToAdd, FragmentsToAdd, TagsToAdd))
+		if (TypedElementDatabasePrivate::ColumnsToBitSets(ColumnsToAdd, FragmentsToAdd, TagsToAdd).MustUpdate())
 		{
 			FMassArchetypeCompositionDescriptor AddComposition(
 				MoveTemp(FragmentsToAdd), MoveTemp(TagsToAdd), FMassChunkFragmentBitSet(), FMassSharedFragmentBitSet());
@@ -490,7 +520,7 @@ bool UTypedElementDatabase::AddRemoveColumns(TypedElementRowHandle Row,
 
 		FMassTagBitSet TagsToRemove;
 		FMassFragmentBitSet FragmentsToRemove;
-		if (ColumnsToBitSets(ColumnsToRemove, FragmentsToRemove, TagsToRemove))
+		if (TypedElementDatabasePrivate::ColumnsToBitSets(ColumnsToRemove, FragmentsToRemove, TagsToRemove).MustUpdate())
 		{
 			FMassArchetypeCompositionDescriptor RemoveComposition(
 				MoveTemp(FragmentsToRemove), MoveTemp(TagsToRemove), FMassChunkFragmentBitSet(), FMassSharedFragmentBitSet());
@@ -503,7 +533,7 @@ bool UTypedElementDatabase::AddRemoveColumns(TypedElementRowHandle Row,
 
 bool UTypedElementDatabase::BatchAddRemoveColumns(TConstArrayView<TypedElementRowHandle> Rows, 
 	TConstArrayView<const UScriptStruct*> ColumnsToAdd, TConstArrayView<const UScriptStruct*> ColumnsToRemove)
-{
+{	
 	if (ActiveEditorEntityManager)
 	{
 		FMassFragmentBitSet FragmentsToAdd;
@@ -512,10 +542,12 @@ bool UTypedElementDatabase::BatchAddRemoveColumns(TConstArrayView<TypedElementRo
 		FMassTagBitSet TagsToAdd;
 		FMassTagBitSet TagsToRemove;
 
-		bool bMustUpdateFragments = ColumnsToBitSets(ColumnsToAdd, FragmentsToAdd, TagsToAdd);
-		bool bMustUpdateTags = ColumnsToBitSets(ColumnsToRemove, FragmentsToRemove, TagsToRemove);
+		namespace TEDP = TypedElementDatabasePrivate;
+
+		TEDP::ColumnsToBitSetsResult AddResult = TEDP::ColumnsToBitSets(ColumnsToAdd, FragmentsToAdd, TagsToAdd);
+		TEDP::ColumnsToBitSetsResult RemoveResult = TEDP::ColumnsToBitSets(ColumnsToRemove, FragmentsToRemove, TagsToRemove);
 		
-		if (bMustUpdateFragments || bMustUpdateTags)
+		if (AddResult.MustUpdate() || RemoveResult.MustUpdate())
 		{
 			using EntityHandleArray = TArray<FMassEntityHandle, TInlineAllocator<32>>;
 			using EntityArchetypeLookup = TMap<FMassArchetypeHandle, EntityHandleArray, TInlineSetAllocator<32>>;
@@ -533,7 +565,7 @@ bool UTypedElementDatabase::BatchAddRemoveColumns(TConstArrayView<TypedElementRo
 					EntityCollection.Add(Entity);
 				}
 			}
-			
+		
 			// Construct table (archetype) specific row (entity) collections.
 			ArchetypeEntityArray EntityCollections;
 			EntityCollections.Reserve(LookupTable.Num());
@@ -543,11 +575,11 @@ bool UTypedElementDatabase::BatchAddRemoveColumns(TConstArrayView<TypedElementRo
 			}
 
 			// Batch update using the appropriate fragment/bit sets.
-			if (bMustUpdateFragments)
+			if (AddResult.bMustUpdateFragments || RemoveResult.bMustUpdateFragments)
 			{
 				ActiveEditorEntityManager->BatchChangeFragmentCompositionForEntities(EntityCollections, FragmentsToAdd, FragmentsToRemove);
 			}
-			if (bMustUpdateTags)
+			if (AddResult.bMustUpdateTags || RemoveResult.bMustUpdateTags)
 			{
 				ActiveEditorEntityManager->BatchChangeTagsForEntities(EntityCollections, TagsToAdd, TagsToRemove);
 			}
@@ -743,25 +775,6 @@ void* UTypedElementDatabase::GetExternalSystemAddress(UClass* Target)
 	return nullptr;
 }
 
-bool UTypedElementDatabase::ColumnsToBitSets(TConstArrayView<const UScriptStruct*> Columns, FMassFragmentBitSet& Fragments, FMassTagBitSet& Tags)
-{
-	bool bResult = false;
-	for (const UScriptStruct* ColumnType : Columns)
-	{
-		if (ColumnType->IsChildOf(FMassFragment::StaticStruct()))
-		{
-			Fragments.Add(*ColumnType);
-			bResult = true;
-		}
-		else if (ColumnType->IsChildOf(FMassTag::StaticStruct()))
-		{
-			Tags.Add(*ColumnType);
-			bResult = true;
-		}
-	}
-	return bResult;
-}
-
 void UTypedElementDatabase::PreparePhase(EQueryTickPhase Phase, float DeltaTime)
 {
 	if (ActiveEditorEntityManager)
diff --git a/Engine/Plugins/Experimental/TypedElementsDataStorage/Source/TypedElementsDataStorage/Private/TypedElementDatabase.h b/Engine/Plugins/Experimental/TypedElementsDataStorage/Source/TypedElementsDataStorage/Private/TypedElementDatabase.h
--- a/Engine/Plugins/Experimental/TypedElementsDataStorage/Source/TypedElementsDataStorage/Private/TypedElementDatabase.h
+++ b/Engine/Plugins/Experimental/TypedElementsDataStorage/Source/TypedElementsDataStorage/Private/TypedElementDatabase.h
@@ -101,9 +101,6 @@ public:
 	void DebugPrintQueryCallbacks(FOutputDevice& Output);
 
 private:
-	/** Converts a set of column types into Mass specific fragment and tag bit sets. Returns true if any values were added. */
-	static bool ColumnsToBitSets(TConstArrayView<const UScriptStruct*> Columns, FMassFragmentBitSet& Fragments, FMassTagBitSet& Tags);
-
 	void PreparePhase(EQueryTickPhase Phase, float DeltaTime);
 	void FinalizePhase(EQueryTickPhase Phase, float DeltaTime);
 	void Reset();
