Skein / ViewportStatus: Fix broken UTypedElementRevisionControlFactory query.

Why? It was no longer working because the database schema had changed.

#rnx

diff --git a/Engine/Plugins/Experimental/TypedElementsDataStorage/Source/TypedElementsDataStorageRevisionControl/Source/Private/RevisionControlProcessors.cpp b/Engine/Plugins/Experimental/TypedElementsDataStorage/Source/TypedElementsDataStorageRevisionControl/Source/Private/RevisionControlProcessors.cpp
--- a/Engine/Plugins/Experimental/TypedElementsDataStorage/Source/TypedElementsDataStorageRevisionControl/Source/Private/RevisionControlProcessors.cpp
+++ b/Engine/Plugins/Experimental/TypedElementsDataStorage/Source/TypedElementsDataStorageRevisionControl/Source/Private/RevisionControlProcessors.cpp
@@ -14,6 +14,8 @@
 #include "Elements/Framework/TypedElementQueryBuilder.h"
 #include "Elements/Framework/TypedElementRegistry.h"
 
+#include "MassActorSubsystem.h"
+
 extern TYPEDELEMENTSDATASTORAGE_API FAutoConsoleVariableRef CVarAutoPopulateState;
 
 FAutoConsoleCommandWithArgsAndOutputDevice SetSelectionSCCStateConsoleCommand(
@@ -33,6 +35,7 @@ FAutoConsoleCommandWithArgsAndOutputDevice SetSelectionSCCStateConsoleCommand(
 				{
 					AllSelectedQuery = DataStorage->RegisterQuery(
 						Select()
+							.ReadOnly<FTypedElementPackageReference>()
 						.Where()
 							.All<FTypedElementSelectionColumn>()
 						.Compile());
@@ -60,17 +63,18 @@ FAutoConsoleCommandWithArgsAndOutputDevice SetSelectionSCCStateConsoleCommand(
 					Modification = static_cast<ESCCModification>(StateIndex);
 				}
 				
-				TArray<TypedElementRowHandle> RowHandles;
+				TArray<TypedElementRowHandle> PackageRowHandles;
 				
-				DataStorage->RunQuery(AllSelectedQuery, [&RowHandles](const DSI::FQueryDescription&, DSI::IDirectQueryContext& Context)
-				{
-					RowHandles = Context.GetRowHandles();
-				});
+				DataStorage->RunQuery(AllSelectedQuery, CreateDirectQueryCallbackBinding(
+					[&PackageRowHandles] (DSI::IDirectQueryContext& Context, const FTypedElementPackageReference& PackageReference)
+					{
+						PackageRowHandles.Add(PackageReference.Row);
+					})
+				);
 
-				for (TypedElementRowHandle Row : RowHandles)
+				for (TypedElementRowHandle Row : PackageRowHandles)
 				{
 					DataStorage->AddOrGetColumn<FSCCStatusColumn>(Row)->Modification = Modification;
-					DataStorage->AddColumn<FTypedElementSyncFromWorldTag>(Row);
 				}
 			}
 		}
@@ -90,20 +94,30 @@ void UTypedElementRevisionControlFactory::RegisterQueries(ITypedElementDataStora
 	using namespace TypedElementQueryBuilder;
 	using DSI = ITypedElementDataStorageInterface;
 
+	TypedElementQueryHandle ObjectToSCCQuery = DataStorage.RegisterQuery(
+		Select()
+			.ReadOnly<FSCCStatusColumn>()
+		.Compile());
+
 	DataStorage.RegisterQuery(
 		Select(
 			TEXT("Change selection outline colors based on SCC status"),
 			// This is in PrePhysics because the outline->actor query is in DuringPhysics and contexts don't flush changes between tick groups
-			FProcessor(DSI::EQueryTickPhase::PrePhysics, DataStorage.GetQueryTickGroupName(DSI::EQueryTickGroups::SyncExternalToDataStorage)),
-			[](DSI::IQueryContext& Context, TypedElementRowHandle Row, const FSCCStatusColumn& SCCStatus)
+			FProcessor(DSI::EQueryTickPhase::PrePhysics, DataStorage.GetQueryTickGroupName(DSI::EQueryTickGroups::SyncExternalToDataStorage)),			
+			[](DSI::IQueryContext& Context, TypedElementRowHandle ObjectRow, const FTypedElementPackageReference& PackageReference)
 			{
-				constexpr int BasicSelectionColorCount = 2;
-				Context.AddColumn<FTypedElementViewportColorColumn>(Row, { .SelectionOutlineColorIndex = static_cast<uint8>(int(SCCStatus.Modification) + BasicSelectionColorCount) });
-				Context.AddColumns<FTypedElementSyncBackToWorldTag>(Row);
+				Context.RunSubquery(0, PackageReference.Row, CreateSubqueryCallbackBinding(
+					[&Context, &ObjectRow](DSI::ISubqueryContext& SubQueryContext, const FSCCStatusColumn& Status)
+					{
+						constexpr int BasicSelectionColorCount = 2;
+						Context.AddColumn<FTypedElementViewportColorColumn>(ObjectRow, { .SelectionOutlineColorIndex = static_cast<uint8>(int(Status.Modification) + BasicSelectionColorCount) });
+						Context.AddColumns<FTypedElementSyncBackToWorldTag>(ObjectRow);
+					})
+				);
 			}
 		)
-		.Where()
-			.All<FTypedElementSyncFromWorldTag>()
+		.DependsOn()
+			.SubQuery(ObjectToSCCQuery)
 		.Compile()
 	);
 	
