Fixed type information columns in TEDS going stale in some cases.

UObject based references in TEDS also store a reference to their type information. If that type information is pointing to script data the type information can go stale after a script compile. This is something that can happen in general with type information, scripts are just the first use case where this was discovered in practice. The solution is to list all the type information pointers in TEDS when a request to list dependencies comes in. This way the type information gets refreshed and is valid again. The approach to refreshing type information through `AddReferencedObjects` caused issues with the garbage collection while the version that used the `OnObjectsReinstanced` version caused order-of-execution errors as the updates wouldn't be processed before objects were reinstanced. The latter solution was still preferred as it's a more targeted solution. To address the order-of-execution issues, the type-to-table lookup is delayed until the objects are processed and added to the data storage.
This change also includes an update to some of the debug commands so they're in a central place, plus some new commands were added.

diff --git a/Engine/Plugins/Experimental/TypedElementsDataStorage/Source/TypedElementsDataStorage/Private/Compatibility/TypedElementActorLabelQueries.cpp b/Engine/Plugins/Experimental/TypedElementsDataStorage/Source/TypedElementsDataStorage/Private/Compatibility/TypedElementActorLabelQueries.cpp
--- a/Engine/Plugins/Experimental/TypedElementsDataStorage/Source/TypedElementsDataStorage/Private/Compatibility/TypedElementActorLabelQueries.cpp
+++ b/Engine/Plugins/Experimental/TypedElementsDataStorage/Source/TypedElementsDataStorage/Private/Compatibility/TypedElementActorLabelQueries.cpp
@@ -7,75 +7,12 @@
 #include "Elements/Columns/TypedElementLabelColumns.h"
 #include "Elements/Framework/TypedElementQueryBuilder.h"
 #include "Elements/Framework/TypedElementRegistry.h"
-#include "HAL/IConsoleManager.h"
 #include "Hash/CityHash.h"
 #include "MassActorSubsystem.h"
-#include "ProfilingDebugging/CpuProfilerTrace.h"
 #include "ScopedTransaction.h"
 
 #define LOCTEXT_NAMESPACE "TypedElementDataStorage"
 
-FAutoConsoleCommandWithOutputDevice PrintActorLabelsConsoleCommand(
-	TEXT("TEDS.PrintActorLabels"),
-	TEXT("Prints out the labels for all actors found in the Typed Elements Data Storage."),
-	FConsoleCommandWithOutputDeviceDelegate::CreateLambda([](FOutputDevice& Output)
-		{
-			using namespace TypedElementQueryBuilder;
-			using DSI = ITypedElementDataStorageInterface;
-
-			TRACE_CPUPROFILER_EVENT_SCOPE(TEDS.PrintActorLabelsCommand);
-
-			if (ITypedElementDataStorageInterface* DataStorage = UTypedElementRegistry::GetInstance()->GetMutableDataStorage())
-			{
-				static TypedElementQueryHandle LabelQuery = TypedElementInvalidQueryHandle;
-				if (LabelQuery == TypedElementInvalidQueryHandle)
-				{
-					LabelQuery = DataStorage->RegisterQuery(
-						Select()
-							.ReadOnly<FMassActorFragment, FTypedElementLabelColumn>()
-						.Compile());
-				}
-				
-				if (LabelQuery != TypedElementInvalidQueryHandle)
-				{
-					FString Message;
-					Output.Log(TEXT("The Typed Elements Data Storage has the following actors:"));
-					DataStorage->RunQuery(LabelQuery, CreateDirectQueryCallbackBinding(
-						[&Output, &Message](DSI::IDirectQueryContext& Context, const FMassActorFragment* Actors, const FTypedElementLabelColumn* Labels)
-						{
-							const uint32 Count = Context.GetRowCount();
-
-							const FTypedElementLabelColumn* LabelsIt = Labels;
-							int32 CharacterCount = 0;
-							// Reserve memory first to avoid repeated memory allocations.
-							for (uint32 Index = 0; Index < Count; ++Index)
-							{
-								CharacterCount 
-									+= 14 /* Prefixed text size */ 
-									+ 16 /* Hex address of actor */
-									+ 3 /* Closing brace, colon and space */
-									+ LabelsIt->Label.Len() 
-									+ 1 /* Trailing new line */;
-								++LabelsIt;
-							}
-							Message.Reset(CharacterCount);
-
-							LabelsIt = Labels;
-							for (uint32 Index = 0; Index < Count; ++Index)
-							{
-								Message.Appendf(TEXT("    Actor (0x%p): %s\n"), Actors->Get(), *LabelsIt->Label);
-
-								++LabelsIt;
-								++Actors;
-							}
-
-							Output.Log(Message);
-						}));
-					Output.Log(TEXT("End of Typed Elements Data Storage actors list."));
-				}
-			}
-		}));
-
 void UTypedElementActorLabelFactory::RegisterQueries(ITypedElementDataStorageInterface& DataStorage)
 {
 	RegisterActorLabelToColumnQuery(DataStorage);
diff --git a/Engine/Plugins/Experimental/TypedElementsDataStorage/Source/TypedElementsDataStorage/Private/Compatibility/TypedElementUObjectPackagePathToColumnQueries.cpp b/Engine/Plugins/Experimental/TypedElementsDataStorage/Source/TypedElementsDataStorage/Private/Compatibility/TypedElementUObjectPackagePathToColumnQueries.cpp
--- a/Engine/Plugins/Experimental/TypedElementsDataStorage/Source/TypedElementsDataStorage/Private/Compatibility/TypedElementUObjectPackagePathToColumnQueries.cpp
+++ b/Engine/Plugins/Experimental/TypedElementsDataStorage/Source/TypedElementsDataStorage/Private/Compatibility/TypedElementUObjectPackagePathToColumnQueries.cpp
@@ -11,7 +11,7 @@
 #include "Misc/Paths.h"
 #include "UObject/Package.h"
 
-static bool bAutoPopulateRevisionControlState = false;
+static bool bAutoPopulateRevisionControlState = true;
 TYPEDELEMENTSDATASTORAGE_API FAutoConsoleVariableRef CVarAutoPopulateState(
 	TEXT("TEDS.RevisionControl.AutoPopulateState"),
 	bAutoPopulateRevisionControlState,
diff --git a/Engine/Plugins/Experimental/TypedElementsDataStorage/Source/TypedElementsDataStorage/Private/Debug/TypedElementDatabaseDebugCommands.cpp b/Engine/Plugins/Experimental/TypedElementsDataStorage/Source/TypedElementsDataStorage/Private/Debug/TypedElementDatabaseDebugCommands.cpp
new file mode 100644
--- /dev/null
+++ b/Engine/Plugins/Experimental/TypedElementsDataStorage/Source/TypedElementsDataStorage/Private/Debug/TypedElementDatabaseDebugCommands.cpp
@@ -0,0 +1,216 @@
+// Copyright Epic Games, Inc. All Rights Reserved.
+
+#include "Elements/Columns/TypedElementCompatibilityColumns.h"
+#include "Elements/Columns/TypedElementLabelColumns.h"
+#include "Elements/Columns/TypedElementTypeInfoColumns.h"
+#include "Elements/Framework/TypedElementQueryBuilder.h"
+#include "Elements/Framework/TypedElementRegistry.h"
+#include "Elements/Interfaces/TypedElementDataStorageInterface.h"
+#include "HAL/IConsoleManager.h"
+#include "MassActorSubsystem.h"
+#include "Misc/OutputDevice.h"
+#include "ProfilingDebugging/CpuProfilerTrace.h"
+
+namespace Private
+{
+	static const UStruct* GetTypeInfo(const TWeakObjectPtr<const UClass>& TypeInfo) { return TypeInfo.Get(); }
+	static const UStruct* GetTypeInfo(const TWeakObjectPtr<const UScriptStruct>& TypeInfo) { return TypeInfo.Get(); }
+	static const UStruct* GetTypeInfo(const TObjectPtr<const UClass>& TypeInfo) { return TypeInfo; }
+	static const UStruct* GetTypeInfo(const TObjectPtr<const UScriptStruct>& TypeInfo) { return TypeInfo; }
+
+	template<typename TypeInfoType>
+	void PrintObjectTypeInformation(ITypedElementDataStorageInterface* DataStorage, FString Message, FOutputDevice& Output)
+	{
+		using namespace TypedElementDataStorage;
+		using namespace TypedElementQueryBuilder;
+
+		static QueryHandle Query = [DataStorage]
+		{
+			return DataStorage->RegisterQuery(
+				Select()
+					.ReadOnly<TypeInfoType>()
+				.Compile());
+		}();
+
+		if (Query != InvalidQueryHandle)
+		{
+			DataStorage->RunQuery(Query, CreateDirectQueryCallbackBinding(
+				[&Output, &Message](IDirectQueryContext& Context, const TypeInfoType* Types)
+				{
+					Message.Reset();
+					Message += TEXT("  Batch start\n");
+
+					TConstArrayView<TypeInfoType> TypeList(Types, Context.GetRowCount());
+					for (const TypeInfoType& Type : TypeList)
+					{
+						if (const UStruct* TypeInfo = GetTypeInfo(Type.TypeInfo))
+						{
+							Message += TEXT("    Type: ");
+							TypeInfo->AppendName(Message);
+							Message += TEXT('\n');
+						}
+						else
+						{
+							Message += TEXT("    Type: [Invalid]\n");
+						}
+					}
+					Message += TEXT("  Batch end\n");
+					Output.Log(Message);
+				}));
+		}
+	}
+
+	template<typename... Conditions>
+	void PrintObjectLabels(FOutputDevice& Output)
+	{
+		using namespace TypedElementQueryBuilder;
+		using DSI = ITypedElementDataStorageInterface;
+
+		if (ITypedElementDataStorageInterface* DataStorage = UTypedElementRegistry::GetInstance()->GetMutableDataStorage())
+		{
+			static TypedElementQueryHandle LabelQuery = [DataStorage]
+			{
+				if constexpr (sizeof...(Conditions) > 0)
+				{
+					return DataStorage->RegisterQuery(
+						Select()
+							.ReadOnly<FTypedElementUObjectColumn, FTypedElementLabelColumn>()
+						.Where()
+							.All<Conditions...>()
+						.Compile());
+				}
+				else
+				{
+					return DataStorage->RegisterQuery(
+						Select()
+							.ReadOnly<FTypedElementUObjectColumn, FTypedElementLabelColumn>()
+						.Compile());
+				}
+			}();
+
+			if (LabelQuery != TypedElementInvalidQueryHandle)
+			{
+				FString Message;
+				DataStorage->RunQuery(LabelQuery, CreateDirectQueryCallbackBinding(
+					[&Output, &Message](DSI::IDirectQueryContext& Context, const FTypedElementUObjectColumn* Objects, const FTypedElementLabelColumn* Labels)
+					{
+						const uint32 Count = Context.GetRowCount();
+
+						const FTypedElementLabelColumn* LabelsIt = Labels;
+						int32 CharacterCount = 2; // Initial blank space and new line.
+						// Reserve memory first to avoid repeated memory allocations.
+						for (uint32 Index = 0; Index < Count; ++Index)
+						{
+							CharacterCount
+								+= 4 /* Indention */
+								+ 16 /* Hex address of actor */
+								+ 2 /* Colon and space */
+								+ LabelsIt->Label.Len()
+								+ 1 /* Trailing new line */;
+							++LabelsIt;
+						}
+						Message.Reset(CharacterCount);
+						Message = TEXT(" \n");
+
+						LabelsIt = Labels;
+						for (uint32 Index = 0; Index < Count; ++Index)
+						{
+							Message.Appendf(TEXT("    0x%p: %s\n"), Objects->Object.Get(), *LabelsIt->Label);
+
+							++LabelsIt;
+							++Objects;
+						}
+
+						Output.Log(Message);
+					}));
+			}
+		}
+	}
+}
+
+FAutoConsoleCommandWithOutputDevice PrintObjectTypeInformationConsoleCommand(
+	TEXT("TEDS.Debug.PrintObjectTypeInfo"),
+	TEXT("Prints the type information of any rows that has a type information column."),
+	FConsoleCommandWithOutputDeviceDelegate::CreateLambda([](FOutputDevice& Output)
+		{
+			TRACE_CPUPROFILER_EVENT_SCOPE(TEDS.Debug.PrintObjectTypeInfo);
+
+			if (ITypedElementDataStorageInterface* DataStorage = UTypedElementRegistry::GetInstance()->GetMutableDataStorage())
+			{
+				FString Message;
+				Output.Log(TEXT("The Typed Elements Data Storage has the types:"));
+				Private::PrintObjectTypeInformation<FTypedElementClassTypeInfoColumn>(DataStorage, Message, Output);
+				Private::PrintObjectTypeInformation<FTypedElementScriptStructTypeInfoColumn>(DataStorage, Message, Output);
+				Output.Log(TEXT("End of Typed Elements Data Storage type list."));
+			}
+		}
+	));
+
+FAutoConsoleCommandWithOutputDevice PrintAllUObjectsLabelsConsoleCommand(
+	TEXT("TEDS.Debug.PrintAllUObjectsLabels"),
+	TEXT("Prints out the labels for all UObjects found in the Typed Elements Data Storage."),
+	FConsoleCommandWithOutputDeviceDelegate::CreateLambda([](FOutputDevice& Output)
+		{
+			TRACE_CPUPROFILER_EVENT_SCOPE(TEDS.Debug.PrintAllUObjectsLabels);
+			Output.Log(TEXT("The Typed Elements Data Storage has the following UObjects:"));
+			Private::PrintObjectLabels(Output);
+			Output.Log(TEXT("End of Typed Elements Data Storage UObjects list."));
+		}));
+
+FAutoConsoleCommandWithOutputDevice PrintActorLabelsConsoleCommand(
+	TEXT("TEDS.Debug.PrintActorLabels"),
+	TEXT("Prints out the labels for all actors found in the Typed Elements Data Storage."),
+	FConsoleCommandWithOutputDeviceDelegate::CreateLambda([](FOutputDevice& Output)
+		{
+			TRACE_CPUPROFILER_EVENT_SCOPE(TEDS.Debug.PrintActorLabels);
+			Output.Log(TEXT("The Typed Elements Data Storage has the following actors:"));
+			Private::PrintObjectLabels<FMassActorFragment>(Output);
+			Output.Log(TEXT("End of Typed Elements Data Storage actors list."));
+		}));
+
+FAutoConsoleCommandWithOutputDevice CheckActorAndObjectAddressCommand(
+	TEXT("TEDS.Debug.CheckActorAndObjectAddress"),
+	TEXT("Checks if the address of the actor is the same as the address of the object."),
+	FConsoleCommandWithOutputDeviceDelegate::CreateLambda([](FOutputDevice& Output)
+		{
+			using namespace TypedElementQueryBuilder;
+			using DSI = ITypedElementDataStorageInterface;
+
+			TRACE_CPUPROFILER_EVENT_SCOPE(TEDS.Debug.CheckActorAndObjectAddress);
+
+			if (ITypedElementDataStorageInterface* DataStorage = UTypedElementRegistry::GetInstance()->GetMutableDataStorage())
+			{
+				static TypedElementQueryHandle LabelQuery = [DataStorage]
+				{
+					return DataStorage->RegisterQuery(
+						Select()
+							.ReadOnly<FTypedElementUObjectColumn, FMassActorFragment>()
+						.Compile());
+				}();
+				
+				if (LabelQuery != TypedElementInvalidQueryHandle)
+				{
+					bool Result = true;
+					DataStorage->RunQuery(LabelQuery, CreateDirectQueryCallbackBinding(
+						[&Result](DSI::IDirectQueryContext& Context, const FTypedElementUObjectColumn* Objects, const FMassActorFragment* Actors)
+						{
+							const uint32 Count = Context.GetRowCount();
+
+							for (uint32 Index = 0; Index < Count; ++Index)
+							{
+								Result = Result && (Objects->Object.Get() == Actors->Get());
+								++Objects;
+								++Actors;
+							}
+						}));
+					if (Result)
+					{
+						Output.Logf(TEXT("All actors in TEDS have matching object addresses."));
+					}
+					else
+					{
+						Output.Logf(TEXT("One or more actors in TEDS have addresses that differ from the object."));
+					}
+				}
+			}
+		}));
diff --git a/Engine/Plugins/Experimental/TypedElementsDataStorage/Source/TypedElementsDataStorage/Private/TypedElementDatabaseCompatibility.cpp b/Engine/Plugins/Experimental/TypedElementsDataStorage/Source/TypedElementsDataStorage/Private/TypedElementDatabaseCompatibility.cpp
--- a/Engine/Plugins/Experimental/TypedElementsDataStorage/Source/TypedElementsDataStorage/Private/TypedElementDatabaseCompatibility.cpp
+++ b/Engine/Plugins/Experimental/TypedElementsDataStorage/Source/TypedElementsDataStorage/Private/TypedElementDatabaseCompatibility.cpp
@@ -2,7 +2,9 @@
 
 #include "TypedElementDatabaseCompatibility.h"
 
+#include <utility>
 #include "Algo/Unique.h"
+#include "Async/UniqueLock.h"
 #include "Editor.h"
 #include "Editor/TransBuffer.h"
 #include "Elements/Columns/TypedElementCompatibilityColumns.h"
@@ -12,6 +14,7 @@
 #include "Elements/Columns/TypedElementTransformColumns.h"
 #include "Elements/Columns/TypedElementTypeInfoColumns.h"
 #include "Elements/Framework/TypedElementIndexHasher.h"
+#include "Elements/Framework/TypedElementQueryBuilder.h"
 #include "MassActorSubsystem.h"
 #include "TypedElementDataStorageProfilingMacros.h"
 
@@ -26,7 +29,8 @@ void UTypedElementDatabaseCompatibility::Initialize(ITypedElementDataStorageInte
 
 	PostEditChangePropertyDelegateHandle = FCoreUObjectDelegates::OnObjectPropertyChanged.AddUObject(this, &UTypedElementDatabaseCompatibility::OnPostEditChangeProperty);
 	ObjectModifiedDelegateHandle = FCoreUObjectDelegates::OnObjectModified.AddUObject(this, &UTypedElementDatabaseCompatibility::OnObjectModified);
-	
+	ObjectReinstancedDelegateHandle = FCoreUObjectDelegates::OnObjectsReinstanced.AddUObject(this, &UTypedElementDatabaseCompatibility::OnObjectReinstanced);
+
 	PostWorldInitializationDelegateHandle = FWorldDelegates::OnPostWorldInitialization.AddUObject(this, &UTypedElementDatabaseCompatibility::OnPostWorldInitialization);
 	PreWorldFinishDestroyDelegateHandle = FWorldDelegates::OnPreWorldFinishDestroy.AddUObject(this, &UTypedElementDatabaseCompatibility::OnPreWorldFinishDestroy);
 }
@@ -40,7 +44,8 @@ void UTypedElementDatabaseCompatibility::Deinitialize()
 
 	FWorldDelegates::OnPreWorldFinishDestroy.Remove(PreWorldFinishDestroyDelegateHandle);
 	FWorldDelegates::OnPostWorldInitialization.Remove(PostWorldInitializationDelegateHandle);
-
+	
+	FCoreUObjectDelegates::OnObjectsReinstanced.Remove(ObjectReinstancedDelegateHandle);
 	FCoreUObjectDelegates::OnObjectModified.Remove(ObjectModifiedDelegateHandle);
 	FCoreUObjectDelegates::OnObjectPropertyChanged.Remove(PostEditChangePropertyDelegateHandle);
 	
@@ -110,13 +115,9 @@ TypedElementRowHandle UTypedElementDatabaseCompatibility::AddCompatibleObjectExp
 		TypedElementRowHandle Result = FindRowWithCompatibleObjectExplicit(Object);
 		if (!Storage->IsRowAvailable(Result))
 		{
-			TableHandle Table = FindBestMatchingTable(TypeInfo.Get());
-			Table = (Table != InvalidTableHandle) ? Table : StandardExternalObjectTable;
-
 			Result = Storage->ReserveRow();
 			Storage->IndexRow(GenerateIndexHash(Object), Result);
-			PendingRegistration<ExternalObjectRegistration>& Pending = ExternalObjectsPendingRegistration.FindOrAdd(Table);
-			Pending.Add(Result, ExternalObjectRegistration{ .Object = Object, .TypeInfo = TypeInfo });
+			ExternalObjectsPendingRegistration.Add(Result, ExternalObjectRegistration{ .Object = Object, .TypeInfo = TypeInfo });
 		}
 		return Result;
 	}
@@ -149,7 +150,6 @@ void UTypedElementDatabaseCompatibility::RemoveCompatibleObjectExplicit(void* Ob
 	}
 }
 
-
 TypedElementRowHandle UTypedElementDatabaseCompatibility::FindRowWithCompatibleObjectExplicit(const UObject* Object) const
 {
 	using namespace TypedElementDataStorage;
@@ -169,19 +169,10 @@ TypedElementRowHandle UTypedElementDatabaseCompatibility::FindRowWithCompatibleO
 	return (Object && Storage && Storage->IsAvailable()) ? Storage->FindIndexedRow(GenerateIndexHash(Object)) : InvalidRowHandle;
 }
 
-void UTypedElementDatabaseCompatibility::AddReferencedObjects(UObject* InThis, FReferenceCollector& Collector)
-{
-	Super::AddReferencedObjects(InThis, Collector);
-
-	UTypedElementDatabaseCompatibility* Compatibility = Cast<UTypedElementDatabaseCompatibility>(InThis);
-	checkf(Compatibility, TEXT("AddReferencedObjects was given a null pointer or an object wasn't a typed element data compatibility object."));
-
-	Collector.AddReferencedObjects(Compatibility->TypeToTableMap);
-}
-
 void UTypedElementDatabaseCompatibility::Prepare()
 {
 	CreateStandardArchetypes();
+	RegisterTypeInformationQueries();
 }
 
 void UTypedElementDatabaseCompatibility::Reset()
@@ -214,6 +205,21 @@ void UTypedElementDatabaseCompatibility::CreateStandardArchetypes()
 	RegisterTypeTableAssociation(UObject::StaticClass(), StandardUObjectTable);
 }
 
+void UTypedElementDatabaseCompatibility::RegisterTypeInformationQueries()
+{
+	using namespace TypedElementQueryBuilder;
+
+	ClassTypeInfoQuery = Storage->RegisterQuery(
+		Select()
+			.ReadWrite<FTypedElementClassTypeInfoColumn>()
+		.Compile());
+	
+	ScriptStructTypeInfoQuery = Storage->RegisterQuery(
+		Select()
+			.ReadWrite<FTypedElementScriptStructTypeInfoColumn>()
+		.Compile());
+}
+
 bool UTypedElementDatabaseCompatibility::ShouldAddObject(const UObject* Object) const
 {
 	using namespace TypedElementDataStorage;
@@ -255,15 +261,9 @@ TypedElementRowHandle UTypedElementDatabaseCompatibility::AddCompatibleObjectExp
 	TypedElementRowHandle Result = FindRowWithCompatibleObjectExplicit(Object);
 	if (!Storage->IsRowAvailable(Result))
 	{
-		TableHandle Table = FindBestMatchingTable(Object->GetClass());
-		checkf(Table != InvalidTableHandle, TEXT("The Typed Elements Data Storage could not find any matching tables for object of type '%s'. "
-			"This can mean that the object doesn't derive from UObject or that a table for UObject is no longer registered."), *Object->GetClass()->GetFName().ToString());
-
 		Result = Storage->ReserveRow();
 		Storage->IndexRow(GenerateIndexHash(Object), Result);
-
-		PendingRegistration<TWeakObjectPtr<UObject>>& Pending = UObjectsPendingRegistration.FindOrAdd(Table);
-		Pending.Add(Result, Object);
+		UObjectsPendingRegistration.Add(Result, Object);
 
 		if constexpr (bEnableTransactions)
 		{
@@ -322,10 +322,11 @@ TypedElementRowHandle UTypedElementDatabaseCompatibility::DealiasObject(const UO
 void UTypedElementDatabaseCompatibility::Tick()
 {
 	TEDS_EVENT_SCOPE(TEXT("Compatibility Tick"))
-	UWorld* EditorWorld = GEditor->GetEditorWorldContext().World();
+	
+	PendingTypeInformationUpdate.Process(*this);
 
 	// Delay processing until the required systems are available by not clearing any lists or doing any work.
-	if (Storage && Storage->IsAvailable() && EditorWorld)
+	if (Storage && Storage->IsAvailable())
 	{
 		TickPendingUObjectRegistration();
 		TickPendingExternalObjectRegistration();
@@ -333,130 +334,278 @@ void UTypedElementDatabaseCompatibility::Tick()
 	}
 }
 
+
+
+//
+// FPendingTypeInformatUpdate
+// 
+
+UTypedElementDatabaseCompatibility::FPendingTypeInformationUpdate::FPendingTypeInformationUpdate()
+	: PendingTypeInformationUpdatesActive(&PendingTypeInformationUpdates[0])
+	, PendingTypeInformationUpdatesSwapped(&PendingTypeInformationUpdates[1])
+{}
+
+void UTypedElementDatabaseCompatibility::FPendingTypeInformationUpdate::AddTypeInformation(const TMap<UObject*, UObject*>& ReplacedObjects)
+{
+	UE::TUniqueLock Lock(Safeguard);
+
+	for (TMap<UObject*, UObject*>::TConstIterator It = ReplacedObjects.CreateConstIterator(); It; ++It)
+	{
+		if (It->Key->IsA<UStruct>())
+		{
+			PendingTypeInformationUpdatesActive->Add(*It);
+			bHasPendingUpdate = true;
+		}
+	}
+}
+
+void UTypedElementDatabaseCompatibility::FPendingTypeInformationUpdate::Process(UTypedElementDatabaseCompatibility& Compatibility)
+{
+	using namespace TypedElementDataStorage;
+	using namespace TypedElementQueryBuilder;
+	
+	if (bHasPendingUpdate)
+	{
+		// Swap to release the lock as soon as possible.
+		{
+			UE::TUniqueLock Lock(Safeguard);
+			std::swap(PendingTypeInformationUpdatesActive, PendingTypeInformationUpdatesSwapped);
+		}
+
+		for (TypeToTableMapType::TIterator It = Compatibility.TypeToTableMap.CreateIterator(); It; ++It)
+		{
+			if (TOptional<TWeakObjectPtr<UObject>> NewObject = ProcessResolveTypeRecursively(It.Key()); NewObject.IsSet())
+			{
+				UpdatedTypeInfoScratchBuffer.Emplace(Cast<UStruct>(*NewObject), It.Value());
+				It.RemoveCurrent();
+			}
+		}
+		for (TPair<TWeakObjectPtr<UStruct>, TypedElementDataStorage::TableHandle>& UpdatedEntry : UpdatedTypeInfoScratchBuffer)
+		{
+			checkf(UpdatedEntry.Key.IsValid(),
+				TEXT("Type info column in data storage has been re-instanced to an object without type information"));
+			Compatibility.TypeToTableMap.Add(UpdatedEntry);
+		}
+		UpdatedTypeInfoScratchBuffer.Reset();
+
+		Compatibility.Storage->RunQuery(Compatibility.ClassTypeInfoQuery, CreateDirectQueryCallbackBinding(
+			[this](IDirectQueryContext& Context, FTypedElementClassTypeInfoColumn& Type)
+			{
+				if (TOptional<TWeakObjectPtr<UObject>> NewObject = ProcessResolveTypeRecursively(Type.TypeInfo); NewObject.IsSet())
+				{
+					Type.TypeInfo = Cast<UClass>(*NewObject);
+					checkf(Type.TypeInfo.IsValid(),
+						TEXT("Type info column in data storage has been re-instanced to an object without class type information"));
+				}
+			}));
+		Compatibility.Storage->RunQuery(Compatibility.ScriptStructTypeInfoQuery, CreateDirectQueryCallbackBinding(
+			[this](IDirectQueryContext& Context, FTypedElementScriptStructTypeInfoColumn& Type)
+			{
+				if (TOptional<TWeakObjectPtr<UObject>> NewObject = ProcessResolveTypeRecursively(Type.TypeInfo); NewObject.IsSet())
+				{
+					Type.TypeInfo = Cast<UScriptStruct>(*NewObject);
+					checkf(Type.TypeInfo.IsValid(),
+						TEXT("Type info column in data storage has been re-instanced to an object without struct type information"));
+				}
+			}));
+
+		Compatibility.ExternalObjectsPendingRegistration.ForEachAddress(
+			[this](ExternalObjectRegistration& Entry)
+			{
+				if (TOptional<TWeakObjectPtr<UObject>> NewObject = ProcessResolveTypeRecursively(Entry.TypeInfo); NewObject.IsSet())
+				{
+					Entry.TypeInfo = Cast<UScriptStruct>(*NewObject);
+					checkf(Entry.TypeInfo.Get(),
+						TEXT("Type info pending processing in data storage has been re-instanced to an object without struct type information"));
+				}
+			});
+
+		PendingTypeInformationUpdatesSwapped->Reset();
+		bHasPendingUpdate = false;
+	}
+}
+
+TOptional<TWeakObjectPtr<UObject>> UTypedElementDatabaseCompatibility::FPendingTypeInformationUpdate::ProcessResolveTypeRecursively(
+	const TWeakObjectPtr<const UObject>& Target)
+{
+	if (const TWeakObjectPtr<UObject>* NewObject = PendingTypeInformationUpdatesSwapped->Find(Target))
+	{
+		TWeakObjectPtr<UObject> LastNewObject = *NewObject;
+		while (const TWeakObjectPtr<UObject>* NextNewObject = PendingTypeInformationUpdatesSwapped->Find(LastNewObject))
+		{
+			LastNewObject = *NextNewObject;
+		}
+		return LastNewObject;
+	}
+	return TOptional<TWeakObjectPtr<UObject>>();
+}
+
+
+
+
+//
+// PendingRegistration
+//
+
 template<typename AddressType>
 void UTypedElementDatabaseCompatibility::PendingRegistration<AddressType>::Add(TypedElementRowHandle ReservedRowHandle, AddressType Address)
 {
-	Addresses.Add(Forward<AddressType>(Address));
-	ReservedRowHandles.Add(ReservedRowHandle);
+	Entries.Emplace(FEntry{ .Address = Address, .Row = ReservedRowHandle });
 }
 
 template<typename AddressType>
 bool UTypedElementDatabaseCompatibility::PendingRegistration<AddressType>::IsEmpty() const
 {
-	// ReservedRowHandles can also be returned as they'll both have the same length.
-	return Addresses.IsEmpty();
+	return Entries.IsEmpty();
 }
 
 template<typename AddressType>
 int32 UTypedElementDatabaseCompatibility::PendingRegistration<AddressType>::Num() const
 {
-	// ReservedRowHandles can also be returned as they'll both have the same length.
-	return Addresses.Num();
+	return Entries.Num();
 }
 
 template<typename AddressType>
-TArrayView<AddressType> UTypedElementDatabaseCompatibility::PendingRegistration<AddressType>::GetAddresses()
+void UTypedElementDatabaseCompatibility::PendingRegistration<AddressType>::ForEachAddress(const TFunctionRef<void(AddressType&)>& Callback)
 {
-	return Addresses;
+	for (FEntry& Entry : Entries)
+	{
+		Callback(Entry.Address);
+	}
 }
 
 template<typename AddressType>
-TArrayView<TypedElementRowHandle> UTypedElementDatabaseCompatibility::PendingRegistration<AddressType>::GetReservedRowHandles()
+void UTypedElementDatabaseCompatibility::PendingRegistration<AddressType>::ProcessEntries(ITypedElementDataStorageInterface& StorageInterface,
+	UTypedElementDatabaseCompatibility& Compatibility, const TFunctionRef<void(TypedElementRowHandle, const AddressType&)>& SetupRowCallback)
 {
-	return ReservedRowHandles;
-}
+	using namespace TypedElementDataStorage;
 
-template<typename AddressType>
-void UTypedElementDatabaseCompatibility::PendingRegistration<AddressType>::RemoveInvalidEntries(
-	ITypedElementDataStorageInterface& StorageInterface, const TFunctionRef<bool(const AddressType&)>& Validator)
-{
-	checkf(Addresses.Num() == ReservedRowHandles.Num(),
-		TEXT("The reserved row handle count (%i) didn't match the stored pointer count (%i)."),
-		ReservedRowHandles.Num(), Addresses.Num());
-
-	AddressType* AddressBegin = Addresses.GetData();
-	AddressType* AddressIt = Addresses.GetData();
-	AddressType* AddressEnd = AddressBegin + Addresses.Num();
-	TypedElementRowHandle* RowHandleBegin = ReservedRowHandles.GetData();
-	TypedElementRowHandle* RowHandleIt = ReservedRowHandles.GetData();
-	while (AddressIt != AddressEnd)
+	// Start by removing any entries that are no longer valid.
+	for (auto It = Entries.CreateIterator(); It; ++It)
 	{
-		if (StorageInterface.IsRowAvailable(*RowHandleIt) && Validator(*AddressIt))
+		bool bIsValid = StorageInterface.IsRowAvailable(It->Row);
+		if constexpr (std::is_same_v<AddressType, TWeakObjectPtr<UObject>>)
 		{
-			++AddressIt;
-			++RowHandleIt;
+			bIsValid = bIsValid && It->Address.IsValid();
+		}
+		else if constexpr (std::is_same_v<AddressType, ExternalObjectRegistration>)
+		{
+			bIsValid = bIsValid && (It->Address.Object != nullptr);
 		}
 		else
 		{
-			// Don't shrink the registration array as the array will be reused with a variety of different object counts.
-			// If memory size becomes an issue it's better to resize the array once after this loop rather than within the
-			// loop to avoid many resizes happening.
-			StorageInterface.RemoveRow(*RowHandleIt);
-			Addresses.RemoveAtSwap(AddressIt - AddressBegin, 1, EAllowShrinking::No);
-			ReservedRowHandles.RemoveAtSwap(RowHandleIt - RowHandleBegin, 1, EAllowShrinking::No);
-			--AddressEnd;
+			static_assert(sizeof(AddressType) == 0, "Unsupported type for pending object registration in data storage compatibility.");
+		}
+
+		if (!bIsValid)
+		{
+			It.RemoveCurrentSwap();
 		}
 	}
-}
 
-template<typename AddressType>
-void UTypedElementDatabaseCompatibility::PendingRegistration<AddressType>::ProcessEntries(ITypedElementDataStorageInterface& StorageInterface,
-	TypedElementTableHandle Table, const TFunctionRef<void(TypedElementRowHandle, const AddressType&)>& SetupRowCallback)
-{
-	if (!IsEmpty())
+	// Check for empty here are the above code could potentially leave an empty array behind. This would result in break the assumption
+	// that there is at least one entry later in this function.
+	if (!Entries.IsEmpty())
 	{
-		AddressType* TargetIt = GetAddresses().GetData();
-		AddressType* TargetEnd = TargetIt + Num();
-		StorageInterface.BatchAddRow(Table, GetReservedRowHandles(), 
-			[this, &SetupRowCallback, &TargetIt, TargetEnd](TypedElementRowHandle Row)
+		// Next resolve the required table handles.
+		for (FEntry& Entry : Entries)
+		{
+			if constexpr (std::is_same_v<AddressType, TWeakObjectPtr<UObject>>)
+			{
+				Entry.Table = Compatibility.FindBestMatchingTable(Entry.Address->GetClass());
+				checkf(Entry.Table != InvalidTableHandle, 
+					TEXT("The data storage could not find any matching tables for object of type '%s'. "
+					"This can mean that the object doesn't derive from UObject or that a table for UObject is no longer registered."), 
+					*Entry.Address->GetClass()->GetFName().ToString());
+
+			}
+			else if constexpr (std::is_same_v<AddressType, ExternalObjectRegistration>)
+			{
+				Entry.Table = Compatibility.FindBestMatchingTable(Entry.Address.TypeInfo.Get());
+				Entry.Table = (Entry.Table != InvalidTableHandle) ? Entry.Table : Compatibility.StandardExternalObjectTable;
+			}
+			else
+			{
+				static_assert(sizeof(AddressType) == 0, "Unsupported type for pending object registration in data storage compatibility.");
+			}
+		}
+
+		// Next sort them by table then by row handle to allow batch insertion.
+		Entries.Sort(
+			[](const FEntry& Lhs, const FEntry& Rhs)
+			{
+				if (Lhs.Table < Rhs.Table)
+				{
+					return true;
+				}
+				else if (Lhs.Table > Rhs.Table)
+				{
+					return false;
+				}
+				else
+				{
+					return Lhs.Row < Rhs.Row;
+				}
+			});
+
+		// Batch up the entries and add them to the storage.
+		FEntry* Current = Entries.GetData();
+		FEntry* End = Current + Entries.Num();
+		
+		FEntry* TableFront = Current;
+		TableHandle CurrentTable = Entries[0].Table;
+		
+		for (; Current != End; ++Current)
+		{
+			if (Current->Table != CurrentTable)
+			{
+				StorageInterface.BatchAddRow(CurrentTable, Compatibility.RowScratchBuffer,
+					[&SetupRowCallback, &TableFront](TypedElementRowHandle Row)
+					{
+						SetupRowCallback(Row, TableFront->Address);
+						++TableFront;
+					});
+
+				CurrentTable = Current->Table;
+				Compatibility.RowScratchBuffer.Reset();
+			}
+			Compatibility.RowScratchBuffer.Add(Current->Row);
+		}
+		StorageInterface.BatchAddRow(CurrentTable, Compatibility.RowScratchBuffer,
+			[&SetupRowCallback, &TableFront](TypedElementRowHandle Row)
 			{
-				SetupRowCallback(Row, *TargetIt);
-				checkf(TargetIt < TargetEnd, TEXT("More (%i) entities were added than were requested (%i)."), TargetEnd - TargetIt, Num());
-				++TargetIt;
+				SetupRowCallback(Row, TableFront->Address);
+				++TableFront;
 			});
+		Compatibility.RowScratchBuffer.Reset();
 	}
 }
 
 template<typename AddressType>
 void UTypedElementDatabaseCompatibility::PendingRegistration<AddressType>::Reset()
 {
-	Addresses.Reset();
-	ReservedRowHandles.Reset();
+	Entries.Reset();
 }
 
 void UTypedElementDatabaseCompatibility::TickPendingUObjectRegistration()
 {
 	if (!UObjectsPendingRegistration.IsEmpty())
 	{
-		// Filter out the objects that are already registered or already destroyed. 
-		for (auto It = UObjectsPendingRegistration.CreateIterator(); It; ++It)
-		{
-			It->Value.RemoveInvalidEntries(*Storage,
-				[](const TWeakObjectPtr<UObject>& Object)
-				{
-					return Object.Get() != nullptr;
-				});
-		}
-
-		// Add the remaining object to the data storage.
-		for (auto It = UObjectsPendingRegistration.CreateIterator(); It; ++It)
-		{
-			It->Value.ProcessEntries(*Storage, It->Key,
-				[Table = It->Key, this](TypedElementRowHandle Row, const TWeakObjectPtr<UObject>& Object)
+		UObjectsPendingRegistration.ProcessEntries(*Storage, *this,
+			[this](TypedElementRowHandle Row, const TWeakObjectPtr<UObject>& Object)
+			{
+				if (AActor* Actor = Cast<AActor>(Object))
 				{
-					if (AActor* Actor = Cast<AActor>(Object))
-					{
-						constexpr bool bIsOwnedByMass = false;
-						Storage->AddOrGetColumn<FMassActorFragment>(Row)->SetNoHandleMapUpdate(FMassEntityHandle::FromNumber(Row), Actor, bIsOwnedByMass);
-					}
+					constexpr bool bIsOwnedByMass = false;
+					Storage->AddOrGetColumn<FMassActorFragment>(Row)->SetNoHandleMapUpdate(FMassEntityHandle::FromNumber(Row), Actor, bIsOwnedByMass);
+				}
 
-					Storage->AddOrGetColumn<FTypedElementUObjectColumn>(Row, FTypedElementUObjectColumn{ .Object = Object });
-					Storage->AddOrGetColumn<FTypedElementClassTypeInfoColumn>(Row, FTypedElementClassTypeInfoColumn{ .TypeInfo = Object->GetClass() });
-					// Make sure the new row is tagged for update.
-					Storage->AddColumn<FTypedElementSyncFromWorldTag>(Row);
-					OnObjectAdded(Object.Get(), Object->GetClass(), Row);
-				});
-		}
+				Storage->AddOrGetColumn<FTypedElementUObjectColumn>(Row, FTypedElementUObjectColumn{ .Object = Object });
+				Storage->AddOrGetColumn<FTypedElementClassTypeInfoColumn>(Row, FTypedElementClassTypeInfoColumn{ .TypeInfo = Object->GetClass() });
+				// Make sure the new row is tagged for update.
+				Storage->AddColumn<FTypedElementSyncFromWorldTag>(Row);
+				OnObjectAdded(Object.Get(), Object->GetClass(), Row);
+			});
 
 		UObjectsPendingRegistration.Reset();
 	}
@@ -466,30 +615,16 @@ void UTypedElementDatabaseCompatibility::TickPendingExternalObjectRegistration()
 {
 	if (!ExternalObjectsPendingRegistration.IsEmpty())
 	{
-		// Filter out the objects that are already registered or already destroyed. 
-		for (auto It = ExternalObjectsPendingRegistration.CreateIterator(); It; ++It)
-		{
-			It->Value.RemoveInvalidEntries(*Storage,
-				[](const ExternalObjectRegistration& Object)
-				{
-					return Object.Object != nullptr;
-				});
-		}
-
-		// Add the remaining object to the data storage.
-		for (auto It = ExternalObjectsPendingRegistration.CreateIterator(); It; ++It)
-		{
-			It->Value.ProcessEntries(*Storage, It->Key, 
-				[Table = It->Key, this](TypedElementRowHandle Row, const ExternalObjectRegistration& Object)
-				{
-					Storage->AddOrGetColumn<FTypedElementExternalObjectColumn>(Row, FTypedElementExternalObjectColumn{ .Object = Object.Object });
-					Storage->AddOrGetColumn<FTypedElementScriptStructTypeInfoColumn>(Row, FTypedElementScriptStructTypeInfoColumn{ .TypeInfo = Object.TypeInfo });
-					// Make sure the new row is tagged for update.
-					Storage->AddColumn<FTypedElementSyncFromWorldTag>(Row);
+		ExternalObjectsPendingRegistration.ProcessEntries(*Storage, *this,
+			[this](TypedElementRowHandle Row, const ExternalObjectRegistration& Object)
+			{
+				Storage->AddOrGetColumn<FTypedElementExternalObjectColumn>(Row, FTypedElementExternalObjectColumn{ .Object = Object.Object });
+				Storage->AddOrGetColumn<FTypedElementScriptStructTypeInfoColumn>(Row, FTypedElementScriptStructTypeInfoColumn{ .TypeInfo = Object.TypeInfo });
+				// Make sure the new row is tagged for update.
+				Storage->AddColumn<FTypedElementSyncFromWorldTag>(Row);
 
-					OnObjectAdded(Object.Object, Object.TypeInfo.Get(), Row);
-				});
-		}
+				OnObjectAdded(Object.Object, Object.TypeInfo.Get(), Row);
+			});
 
 		ExternalObjectsPendingRegistration.Reset();
 	}
@@ -517,7 +652,8 @@ void UTypedElementDatabaseCompatibility::TickObjectSync()
 					}
 				}
 				const int32 RowHandleCount = RowHandleIndex;
-				RowHandles.SetNum(RowHandleCount, EAllowShrinking::No);
+				const bool bAllowShrinking = false;
+				RowHandles.SetNum(RowHandleCount, bAllowShrinking);
 			}
 
 			ObjectsNeedingFullSync.Reset();
@@ -571,6 +707,11 @@ void UTypedElementDatabaseCompatibility::OnPreObjectRemoved(const void* Object,
 	}
 }
 
+void UTypedElementDatabaseCompatibility::OnObjectReinstanced(const FCoreUObjectDelegates::FReplacementObjectMap& ReplacedObjects)
+{
+	PendingTypeInformationUpdate.AddTypeInformation(ReplacedObjects);
+}
+
 void UTypedElementDatabaseCompatibility::OnPostWorldInitialization(UWorld* World, const UWorld::InitializationValues InitializationValues)
 {
 	FDelegateHandle Handle = World->AddOnActorDestroyedHandler(
diff --git a/Engine/Plugins/Experimental/TypedElementsDataStorage/Source/TypedElementsDataStorage/Private/TypedElementDatabaseCompatibility.h b/Engine/Plugins/Experimental/TypedElementsDataStorage/Source/TypedElementsDataStorage/Private/TypedElementDatabaseCompatibility.h
--- a/Engine/Plugins/Experimental/TypedElementsDataStorage/Source/TypedElementsDataStorage/Private/TypedElementDatabaseCompatibility.h
+++ b/Engine/Plugins/Experimental/TypedElementsDataStorage/Source/TypedElementsDataStorage/Private/TypedElementDatabaseCompatibility.h
@@ -2,6 +2,8 @@
 
 #pragma once
 
+#include <atomic>
+#include "Async/Mutex.h"
 #include "Containers/Array.h"
 #include "Containers/Map.h"
 #include "Compatibility/TypedElementObjectReinstancingManager.h"
@@ -9,6 +11,7 @@
 #include "Elements/Interfaces/TypedElementDataStorageCompatibilityInterface.h"
 #include "Engine/World.h"
 #include "Misc/Change.h"
+#include "Misc/Optional.h"
 #include "UObject/ObjectMacros.h"
 #include "UObject/WeakObjectPtr.h"
 
@@ -44,16 +47,14 @@ public:
 	void UnregisterObjectRemovedCallback(FDelegateHandle Handle);
 	
 	TypedElementRowHandle AddCompatibleObjectExplicit(UObject* Object) override;
-	TypedElementRowHandle AddCompatibleObjectExplicit(void* Object, TWeakObjectPtr<const UScriptStruct> TypeInfo);
+	TypedElementRowHandle AddCompatibleObjectExplicit(void* Object, TWeakObjectPtr<const UScriptStruct> TypeInfo) override;
 	
 	void RemoveCompatibleObjectExplicit(UObject* Object) override;
 	void RemoveCompatibleObjectExplicit(void* Object) override;
 
 	TypedElementRowHandle FindRowWithCompatibleObjectExplicit(const UObject* Object) const override;
 	TypedElementRowHandle FindRowWithCompatibleObjectExplicit(const void* Object) const override;
-
-	static void AddReferencedObjects(UObject* InThis, FReferenceCollector& Collector);
-
+	
 private:
 
 	// The below changes expect UTypedElementDatabaseCompatibility to be the object passed in to StoreUndo.
@@ -87,6 +88,7 @@ private:
 	void Prepare();
 	void Reset();
 	void CreateStandardArchetypes();
+	void RegisterTypeInformationQueries();
 	
 	bool ShouldAddObject(const UObject* Object) const;
 	TypedElementDataStorage::TableHandle FindBestMatchingTable(const UStruct* TypeInfo) const;
@@ -105,42 +107,75 @@ private:
 	void OnObjectModified(UObject* Object);
 	void OnObjectAdded(const void* Object, FTypedElementDatabaseCompatibilityObjectTypeInfo TypeInfo, TypedElementRowHandle Row) const;
 	void OnPreObjectRemoved(const void* Object, FTypedElementDatabaseCompatibilityObjectTypeInfo TypeInfo, TypedElementRowHandle Row) const;
-	
+	void OnObjectReinstanced(const FCoreUObjectDelegates::FReplacementObjectMap& ReplacedObjects);
+
 	void OnPostWorldInitialization(UWorld* World, const UWorld::InitializationValues InitializationValues);
 	void OnPreWorldFinishDestroy(UWorld* World);
 
 	void OnActorDestroyed(AActor* Actor);
+
+	struct FPendingTypeInformationUpdate
+	{
+	public:
+		FPendingTypeInformationUpdate();
+
+		void AddTypeInformation(const TMap<UObject*, UObject*>& ReplacedObjects);
+		void Process(UTypedElementDatabaseCompatibility& Compatibility);
+
+	private:
+		TOptional<TWeakObjectPtr<UObject>> ProcessResolveTypeRecursively(const TWeakObjectPtr<const UObject>& Target);
+
+		struct FTypeInfoEntryKeyFuncs : TDefaultMapHashableKeyFuncs<TWeakObjectPtr<UObject>, TWeakObjectPtr<UObject>, false>
+		{
+			static inline bool Matches(KeyInitType Lhs, KeyInitType Rhs) { return Lhs.HasSameIndexAndSerialNumber(Rhs); }
+		};
+		using PendingTypeInformationMap = TMap<TWeakObjectPtr<UObject>, TWeakObjectPtr<UObject>, FDefaultSetAllocator, FTypeInfoEntryKeyFuncs>;
+
+		PendingTypeInformationMap PendingTypeInformationUpdates[2];
+		PendingTypeInformationMap* PendingTypeInformationUpdatesActive;
+		PendingTypeInformationMap* PendingTypeInformationUpdatesSwapped;
+		TArray<TTuple<TWeakObjectPtr<UStruct>, TypedElementDataStorage::TableHandle>> UpdatedTypeInfoScratchBuffer;
+		UE::FMutex Safeguard;
+		std::atomic<bool> bHasPendingUpdate = false;
+	};
+	FPendingTypeInformationUpdate PendingTypeInformationUpdate;
+
+	struct ExternalObjectRegistration
+	{
+		void* Object;
+		TWeakObjectPtr<const UScriptStruct> TypeInfo;
+	};
 	
 	template<typename AddressType>
 	struct PendingRegistration
 	{
 	private:
-		TArray<AddressType> Addresses;
-		TArray<TypedElementRowHandle> ReservedRowHandles;
+		struct FEntry
+		{
+			AddressType Address;
+			TypedElementDataStorage::RowHandle Row;
+			TypedElementDataStorage::TableHandle Table;
+		};
+		TArray<FEntry> Entries;
 
 	public:
 		void Add(TypedElementRowHandle ReservedRowHandle, AddressType Address);
 		bool IsEmpty() const;
 		int32 Num() const;
-		TArrayView<AddressType> GetAddresses();
-		TArrayView<TypedElementRowHandle> GetReservedRowHandles();
-
-		void RemoveInvalidEntries(ITypedElementDataStorageInterface& Storage, const TFunctionRef<bool(const AddressType&)>& Validator);
-		void ProcessEntries(ITypedElementDataStorageInterface& Storage, TypedElementTableHandle Table,
+		
+		void ForEachAddress(const TFunctionRef<void(AddressType&)>& Callback);
+		void ProcessEntries(ITypedElementDataStorageInterface& Storage, UTypedElementDatabaseCompatibility& Compatibility,
 			const TFunctionRef<void(TypedElementRowHandle, const AddressType&)>& SetupRowCallback);
 		void Reset();
 	};
-	struct ExternalObjectRegistration
-	{
-		void* Object;
-		TWeakObjectPtr<const UScriptStruct> TypeInfo;
-	};
-	TMap<TypedElementTableHandle, PendingRegistration<TWeakObjectPtr<UObject>>> UObjectsPendingRegistration;
-	TMap<TypedElementTableHandle, PendingRegistration<ExternalObjectRegistration>> ExternalObjectsPendingRegistration;
+	PendingRegistration<TWeakObjectPtr<UObject>> UObjectsPendingRegistration;
+	PendingRegistration<ExternalObjectRegistration> ExternalObjectsPendingRegistration;
+	TArray<TypedElementDataStorage::RowHandle> RowScratchBuffer;
 	
 	TArray<ObjectRegistrationFilter> ObjectRegistrationFilters;
 	TArray<ObjectToRowDealiaser> ObjectToRowDialiasers;
-	TMap<TObjectPtr<UStruct>, TypedElementDataStorage::TableHandle> TypeToTableMap;
+	using TypeToTableMapType = TMap<TWeakObjectPtr<UStruct>, TypedElementDataStorage::TableHandle>;
+	TypeToTableMapType TypeToTableMap;
 	TArray<TPair<ObjectAddedCallback, FDelegateHandle>> ObjectAddedCallbackList;
 	TArray<TPair<ObjectRemovedCallback, FDelegateHandle>> PreObjectRemovedCallbackList;
 
@@ -161,6 +196,10 @@ private:
 	FDelegateHandle ObjectModifiedDelegateHandle;
 	FDelegateHandle PostWorldInitializationDelegateHandle;
 	FDelegateHandle PreWorldFinishDestroyDelegateHandle;
+	FDelegateHandle ObjectReinstancedDelegateHandle;
+
+	TypedElementDataStorage::QueryHandle ClassTypeInfoQuery;
+	TypedElementDataStorage::QueryHandle ScriptStructTypeInfoQuery;
 };
 
 enum class ETypedElementDatabaseCompatibilityObjectType : uint8
