Fixed SCC feature that wasn't supposed to be enabled yet.

diff --git a/Engine/Plugins/Experimental/TypedElementsDataStorage/Source/TypedElementsDataStorage/Private/Compatibility/TypedElementUObjectPackagePathToColumnQueries.cpp b/Engine/Plugins/Experimental/TypedElementsDataStorage/Source/TypedElementsDataStorage/Private/Compatibility/TypedElementUObjectPackagePathToColumnQueries.cpp
--- a/Engine/Plugins/Experimental/TypedElementsDataStorage/Source/TypedElementsDataStorage/Private/Compatibility/TypedElementUObjectPackagePathToColumnQueries.cpp
+++ b/Engine/Plugins/Experimental/TypedElementsDataStorage/Source/TypedElementsDataStorage/Private/Compatibility/TypedElementUObjectPackagePathToColumnQueries.cpp
@@ -11,7 +11,7 @@
 #include "Misc/Paths.h"
 #include "UObject/Package.h"
 
-static bool bAutoPopulateRevisionControlState = true;
+static bool bAutoPopulateRevisionControlState = false;
 TYPEDELEMENTSDATASTORAGE_API FAutoConsoleVariableRef CVarAutoPopulateState(
 	TEXT("TEDS.RevisionControl.AutoPopulateState"),
 	bAutoPopulateRevisionControlState,
