Rewrite of TEDS' Scratch Buffer.

This change addresses a memory stomp that happened, which was caused by an incorrect use of placement new.
Originally the issue was thought to be with the Scratch Buffer's custom allocator so the choice was made to move to FConcurrentLinearBulkObjectAllocator. This has been kept in to reduce future maintenance work. The implementation of `AddColumnUnitialized` has also been updated to remove the near duplicate implemenation.

diff --git a/Engine/Plugins/Experimental/TypedElementsDataStorage/Source/TypedElementsDataStorage/Private/Processors/TypedElementProcessorAdaptors.cpp b/Engine/Plugins/Experimental/TypedElementsDataStorage/Source/TypedElementsDataStorage/Private/Processors/TypedElementProcessorAdaptors.cpp
--- a/Engine/Plugins/Experimental/TypedElementsDataStorage/Source/TypedElementsDataStorage/Private/Processors/TypedElementProcessorAdaptors.cpp
+++ b/Engine/Plugins/Experimental/TypedElementsDataStorage/Source/TypedElementsDataStorage/Private/Processors/TypedElementProcessorAdaptors.cpp
@@ -304,105 +304,27 @@ public:
 
 	void* AddColumnUninitialized(TypedElementDataStorage::RowHandle Row, const UScriptStruct* ObjectType) override
 	{
-		checkf(ObjectType->IsChildOf(FMassTag::StaticStruct()) || ObjectType->IsChildOf(FMassFragment::StaticStruct()), TEXT("Column [%s] must be a Mass Fragment or Tag"), *ObjectType->GetName());
-		
-		struct FAddValueColumn
-		{
-			const UScriptStruct* FragmentType;
-			FMassEntityHandle Entity;
-			void* Object;
-
-			FAddValueColumn() = default;
-			FAddValueColumn(const UScriptStruct* InFragmentType, FMassEntityHandle InEntity, void* InObject)
-				: FragmentType(InFragmentType)
-				, Entity(InEntity)
-				, Object(InObject)
-			{}
-		};
-
-		struct FAddValueColumnWithDestructor : FAddValueColumn
-		{
-			FAddValueColumnWithDestructor(const UScriptStruct* InFragmentType, FMassEntityHandle InEntity, void* InObject)
-				: FAddValueColumn(InFragmentType, InEntity, InObject)
-			{}
-
-			~FAddValueColumnWithDestructor()
+		return AddColumnUninitialized(Row, ObjectType,
+			[](const UScriptStruct* TypeInfo, void* Destination, void* Source)
 			{
-				FragmentType->DestroyStruct(Object);
-			}
-		};
-
-		FTypedElementDatabaseScratchBuffer& ScratchBuffer = Environment.GetScratchBuffer();
-		void* ObjectCopy = ScratchBuffer.Allocate(ObjectType->GetStructureSize(), ObjectType->GetMinAlignment());
-		FAddValueColumn* AddedColumn = nullptr;
-		if (ObjectType->StructFlags & (STRUCT_IsPlainOldData | STRUCT_NoDestructor))
-		{
-			AddedColumn = ScratchBuffer.Emplace<FAddValueColumn>(ObjectType, FMassEntityHandle::FromNumber(Row), ObjectCopy);
-		}
-		else
-		{
-			AddedColumn = ScratchBuffer.Emplace<FAddValueColumnWithDestructor>(ObjectType, FMassEntityHandle::FromNumber(Row), ObjectCopy);
-		}
-
-		if (ObjectType->IsChildOf(FMassFragment::StaticStruct()))
-		{
-			Context.Defer().PushCommand<FMassDeferredAddCommand>(
-				[AddedColumn](FMassEntityManager& System)
-				{
-					// Check entity before proceeding. It's possible it may have been invalidated before this defer call fired
-					if (System.IsEntityActive(AddedColumn->Entity))
-					{
-						// Check before adding.  Mass's AddFragmentToEntity is not idempotent and will assert if adding
-						// column to a row that already has one
-						FStructView Fragment = System.GetFragmentDataStruct(AddedColumn->Entity, AddedColumn->FragmentType);
-						if (!Fragment.IsValid())
-						{
-							System.AddFragmentToEntity(AddedColumn->Entity, AddedColumn->FragmentType);
-						}
-					}
-				});
-			
-			Context.Defer().PushCommand<FMassDeferredSetCommand>(
-				[AddedColumn](FMassEntityManager& System)
-				{
-					// Check entity before proceeding. It's possible it may have been invalidated before this defer call fired
-					if (System.IsEntityActive(AddedColumn->Entity))
-					{
-						FStructView Fragment = System.GetFragmentDataStruct(AddedColumn->Entity, AddedColumn->FragmentType);
-						AddedColumn->FragmentType->CopyScriptStruct(Fragment.GetMemory(), AddedColumn->Object);
-					}
-				});
-		}
-		else if (ObjectType->IsChildOf(FMassTag::StaticStruct()))
-		{
-			Context.Defer().PushCommand<FMassDeferredAddCommand>(
-				[AddedColumn](FMassEntityManager& System)
-				{
-					// Check entity before proceeding. It's possible it may have been invalidated before this defer call fired
-					if (System.IsEntityActive(AddedColumn->Entity))
-					{
-						System.AddTagToEntity(AddedColumn->Entity, AddedColumn->FragmentType);
-					}
-				});
-		}
-
-		return ObjectCopy;
+				TypeInfo->CopyScriptStruct(Destination, Source);
+			});
 	}
 
-	void* AddColumnUninitialized(TypedElementDataStorage::RowHandle Row, const UScriptStruct* ObjectType, ObjectMoveOperator Mover) override
+	void* AddColumnUninitialized(TypedElementDataStorage::RowHandle Row, const UScriptStruct* ObjectType, ObjectCopyOrMove Relocator) override
 	{
-		checkf(ObjectType->IsChildOf(FMassTag::StaticStruct()) || ObjectType->IsChildOf(FMassFragment::StaticStruct()), TEXT("Column [%s] must be a Mass Fragment or Tag"), *ObjectType->GetName());
+		checkf(ObjectType->IsChildOf(FMassFragment::StaticStruct()), TEXT("Column [%s] can not be a tag"), *ObjectType->GetName());
 		
 		struct FAddMoveableValueColumn
 		{
-			ObjectMoveOperator Mover;
+			ObjectCopyOrMove Relocator;
 			const UScriptStruct* FragmentType;
 			FMassEntityHandle Entity;
 			void* Object;
 
 			FAddMoveableValueColumn() = default;
-			FAddMoveableValueColumn(ObjectMoveOperator InMover, const UScriptStruct* InFragmentType, FMassEntityHandle InEntity, void* InObject)
-				: Mover(InMover)
+			FAddMoveableValueColumn(ObjectCopyOrMove InRelocator, const UScriptStruct* InFragmentType, FMassEntityHandle InEntity, void* InObject)
+				: Relocator(InRelocator)
 				, FragmentType(InFragmentType)
 				, Entity(InEntity)
 				, Object(InObject)
@@ -412,8 +334,8 @@ public:
 		struct FAddMoveableValueColumnWithDestructor : FAddMoveableValueColumn
 		{
 			FAddMoveableValueColumnWithDestructor(
-				ObjectMoveOperator InMover, const UScriptStruct* InFragmentType, FMassEntityHandle InEntity, void* InObject)
-				: FAddMoveableValueColumn(InMover, InFragmentType, InEntity, InObject)
+				ObjectCopyOrMove InRelocator, const UScriptStruct* InFragmentType, FMassEntityHandle InEntity, void* InObject)
+				: FAddMoveableValueColumn(InRelocator, InFragmentType, InEntity, InObject)
 			{}
 
 			~FAddMoveableValueColumnWithDestructor()
@@ -423,48 +345,37 @@ public:
 		};
 
 		FTypedElementDatabaseScratchBuffer& ScratchBuffer = Environment.GetScratchBuffer();
-		void* MovedObject = ScratchBuffer.Allocate(ObjectType->GetStructureSize(), ObjectType->GetMinAlignment());
+		void* ColumnData = ScratchBuffer.Allocate(ObjectType->GetStructureSize(), ObjectType->GetMinAlignment());
 		FAddMoveableValueColumn* AddedColumn = nullptr;
 		if (ObjectType->StructFlags & (STRUCT_IsPlainOldData | STRUCT_NoDestructor))
 		{
-			AddedColumn = ScratchBuffer.Emplace<FAddMoveableValueColumn>(Mover, ObjectType, FMassEntityHandle::FromNumber(Row), MovedObject);
+			AddedColumn = ScratchBuffer.Emplace<FAddMoveableValueColumn>(Relocator, ObjectType, FMassEntityHandle::FromNumber(Row), ColumnData);
 		}
 		else
 		{
-			AddedColumn = ScratchBuffer.Emplace<FAddMoveableValueColumnWithDestructor>(Mover, ObjectType, FMassEntityHandle::FromNumber(Row), MovedObject);
+			AddedColumn = ScratchBuffer.Emplace<FAddMoveableValueColumnWithDestructor>(Relocator, ObjectType, FMassEntityHandle::FromNumber(Row), ColumnData);
 		}
 
-		if (ObjectType->IsChildOf(FMassFragment::StaticStruct()))
-		{
-			Context.Defer().PushCommand<FMassDeferredAddCommand>(
-				[AddedColumn](FMassEntityManager& System)
-				{
-					// Check before adding.  Mass's AddFragmentToEntity is not idempotent and will assert if adding
-					// column to a row that already has one
-					FStructView Fragment = System.GetFragmentDataStruct(AddedColumn->Entity, AddedColumn->FragmentType);
-					if (!Fragment.IsValid())
-					{
-						System.AddFragmentToEntity(AddedColumn->Entity, AddedColumn->FragmentType);
-					}
-				});
-
-			Context.Defer().PushCommand<FMassDeferredSetCommand>(
-				[AddedColumn](FMassEntityManager& System)
-				{
-					FStructView Fragment = System.GetFragmentDataStruct(AddedColumn->Entity, AddedColumn->FragmentType);
-					AddedColumn->Mover(Fragment.GetMemory(), AddedColumn->Object);
-				});
-		}
-		else if (ObjectType->IsChildOf(FMassTag::StaticStruct()))
-		{
-			Context.Defer().PushCommand<FMassDeferredAddCommand>(
-				[AddedColumn](FMassEntityManager& System)
+		Context.Defer().PushCommand<FMassDeferredAddCommand>(
+			[AddedColumn](FMassEntityManager& System)
+			{
+				// Check before adding.  Mass's AddFragmentToEntity is not idempotent and will assert if adding
+				// column to a row that already has one
+				FStructView Fragment = System.GetFragmentDataStruct(AddedColumn->Entity, AddedColumn->FragmentType);
+				if (!Fragment.IsValid())
 				{
-					System.AddTagToEntity(AddedColumn->Entity, AddedColumn->FragmentType);
-				});
-		}
+					System.AddFragmentToEntity(AddedColumn->Entity, AddedColumn->FragmentType);
+				}
+			});
 
-		return MovedObject;
+		Context.Defer().PushCommand<FMassDeferredSetCommand>(
+			[AddedColumn](FMassEntityManager& System)
+			{
+				FStructView Fragment = System.GetFragmentDataStruct(AddedColumn->Entity, AddedColumn->FragmentType);
+				AddedColumn->Relocator(AddedColumn->FragmentType, Fragment.GetMemory(), AddedColumn->Object);
+			});
+		
+		return ColumnData;
 	}
 
 	void AddColumns(TypedElementDataStorage::RowHandle Row, TConstArrayView<const UScriptStruct*> ColumnTypes) override
diff --git a/Engine/Plugins/Experimental/TypedElementsDataStorage/Source/TypedElementsDataStorage/Private/TypedElementDatabase.cpp b/Engine/Plugins/Experimental/TypedElementsDataStorage/Source/TypedElementsDataStorage/Private/TypedElementDatabase.cpp
--- a/Engine/Plugins/Experimental/TypedElementsDataStorage/Source/TypedElementsDataStorage/Private/TypedElementDatabase.cpp
+++ b/Engine/Plugins/Experimental/TypedElementsDataStorage/Source/TypedElementsDataStorage/Private/TypedElementDatabase.cpp
@@ -221,16 +221,13 @@ void UTypedElementDatabase::OnPreMassTick(float DeltaTime)
 	// to complete pending work.
 	FTypedElementDatabaseCommandBuffer::ProcessBuffer(DeferredCommands, *Environment);
 	DeferredCommands.Reset();
-
-	Environment->GetScratchBuffer().NextFrame();
 }
 
 void UTypedElementDatabase::OnPostMassTick(float DeltaTime)
 {
 	checkf(IsAvailable(), TEXT("Typed Element Database was ticked while it's not ready."));
 	
-	// Recycle any full scratch memory blocks from the previous frame.
-	Environment->GetScratchBuffer().RecycleBlocks();
+	Environment->GetScratchBuffer().BatchDelete();
 }
 
 TSharedPtr<FMassEntityManager> UTypedElementDatabase::GetActiveMutableEditorEntityManager()
diff --git a/Engine/Plugins/Experimental/TypedElementsDataStorage/Source/TypedElementsDataStorage/Private/TypedElementDatabaseScratchBuffer.cpp b/Engine/Plugins/Experimental/TypedElementsDataStorage/Source/TypedElementsDataStorage/Private/TypedElementDatabaseScratchBuffer.cpp
--- a/Engine/Plugins/Experimental/TypedElementsDataStorage/Source/TypedElementsDataStorage/Private/TypedElementDatabaseScratchBuffer.cpp
+++ b/Engine/Plugins/Experimental/TypedElementsDataStorage/Source/TypedElementsDataStorage/Private/TypedElementDatabaseScratchBuffer.cpp
@@ -2,301 +2,22 @@
 
 #include "TypedElementDatabaseScratchBuffer.h"
 
-// Set to 1 to run additional tests on the scratch buffer. These are normally to expensive to run, but can help find deeper issues with the
-// scratch buffer.
-#define TEDS_SCRATCHBUFFER_ENABLE_ADDITIONAL_TESTS 0
-
-FTypedElementDatabaseScratchBuffer::~FTypedElementDatabaseScratchBuffer()
+FTypedElementDatabaseScratchBuffer::FTypedElementDatabaseScratchBuffer()
+	: CurrentAllocator(&Allocators[0])
+	, PreviousAllocator(&Allocators[1])
+	, LeastRecentAllocator(&Allocators[2])
 {
-	// This assumes it's being called after all threads using the scratch buffer have been shutdown already so the
-	// scratch buffers they hold have been released back into the pool.
-
-	NextFrame(); // Force to next frame so any lingering blocks are freed correctly.
-	RecycleBlocks();
-	checkf(FullBlocks == nullptr, TEXT("Not all full blocks were recycled. This can lead to memory leaks."));
-	
-	FBlock* FrontAvailable = AvailableBlocks;
-	while (FrontAvailable)
-	{
-		FBlock* Next = FrontAvailable->NextBlock;
-		delete FrontAvailable;
-		FrontAvailable = Next;
-	}
-	AvailableBlocks = nullptr;
 }
 
 void* FTypedElementDatabaseScratchBuffer::Allocate(size_t Size, size_t Alignment)
 {
-	// ReaderAccessDetector to assert that this scope has shared access to any thread also making an Allocate call
-	UE_MT_SCOPED_READ_ACCESS(AccessDetector);
-	return GetThreadLocalBlockController().Allocate(Size, Alignment, FrameId);
-}
-
-void FTypedElementDatabaseScratchBuffer::NextFrame()
-{
-	FrameId++;
-}
-
-void FTypedElementDatabaseScratchBuffer::RecycleBlocks()
-{
-	if (FullBlocks)
-	{
-		// WriterAccessDetector to assert that this scope has exclusive access - should not overlap Allocate()
-		UE_MT_SCOPED_WRITE_ACCESS(AccessDetector);
-		
-		// Destroy any objects in blocks. It walks the allocations in a block backwards and invokes the destructors on the allocated 
-		// objects. It moves to the next block once it reaches the beginning of the block. At the end the cleared blocks are reinserted 
-		// into the available blocks for reuse. This approach allows for simple lock-free appending of the destructor chain while remaining
-		// relatively simple to destroy. The dirty blocks are added in front of the available blocks instead of the end to avoid having
-		// to iterate over the entire list of available blocks.
-
-#if TEDS_SCRATCHBUFFER_ENABLE_ADDITIONAL_TESTS
-		// Check to see if there are any cycles.
-		{
-			TSet<FBlock*> CycleDetector;
-			FBlock* TestBlock = FullBlocks;
-			while (TestBlock)
-			{
-				checkf(CycleDetector.Find(TestBlock) == nullptr, TEXT("A circular dependency was found in scratch buffer."));
-				CycleDetector.Add(TestBlock);
-				TestBlock = TestBlock->NextBlock;
-			}
-		}
-		{
-			TSet<FBlock*> CycleDetector;
-			FBlock* TestBlock = AvailableBlocks;
-			while (TestBlock)
-			{
-				checkf(CycleDetector.Find(TestBlock) == nullptr, TEXT("A circular dependency was found in scratch buffer."));
-				CycleDetector.Add(TestBlock);
-				TestBlock = TestBlock->NextBlock;
-			}
-		}
-#endif // TEDS_SCRATCHBUFFER_ENABLE_ADDITIONAL_TESTS
-
-		FBlock* RemainingBlocks = nullptr;
-		FBlock* DirtyBlock = FullBlocks;
-		do
-		{
-			// If the block hasn't been touched for at least one frame it can be recycled.
-			if (DirtyBlock->LastTouchedByFrame < FrameId)
-			{
-				// Call destructors on any memory that needs it.
-				FDestructorTail* Destructor = DirtyBlock->DestructionTail;
-				while (Destructor)
-				{
-					Destructor->Destructor(reinterpret_cast<char*>(Destructor) - Destructor->StructOffset, Destructor->InstanceCount);
-					Destructor = Destructor->PreviousTail;
-				}
-				DirtyBlock->DestructionTail = nullptr;
-
-				// The block is now clean so set the next block to process.
-				FBlock* CleanBlock = DirtyBlock;
-				DirtyBlock = DirtyBlock->NextBlock;
-				
-				// Reinsert the clean block into the chain of available blocks.
-				checkf(AvailableBlocks != CleanBlock, TEXT("Recycled block has already been added to the available blocks."));
-				CleanBlock->NextBlock.store(AvailableBlocks);
-				AvailableBlocks = CleanBlock;
-			}
-			else
-			{
-				// The block can't be processed yet so move onto the next block.
-				FBlock* DelayedBlock = DirtyBlock;
-				DirtyBlock = DirtyBlock->NextBlock;
-
-				// Store the delayed blocks for future processing.
-				DelayedBlock->NextBlock = RemainingBlocks;
-				RemainingBlocks = DelayedBlock;
-			}
-		} while (DirtyBlock);
-		
-		// Set the full blocks to the first block that couldn't be removed or null if all filled up blocks have been recycled.
-		FullBlocks = RemainingBlocks;
-
-#if TEDS_SCRATCHBUFFER_ENABLE_ADDITIONAL_TESTS
-		// Check if no new cycles have been introduced after updating.
-		{
-			TSet<FBlock*> CycleDetector;
-			FBlock* TestBlock = FullBlocks;
-			while (TestBlock)
-			{
-				checkf(CycleDetector.Find(TestBlock) == nullptr, TEXT("A circular dependency was found in scratch buffer."));
-				checkf(TestBlock->LastTouchedByFrame >= FrameId, TEXT("Deleted block still referenced."));
-				CycleDetector.Add(TestBlock);
-				TestBlock = TestBlock->NextBlock;
-			}
-		}
-		{
-			TSet<FBlock*> CycleDetector;
-			FBlock* TestBlock = AvailableBlocks;
-			while (TestBlock)
-			{
-				checkf(CycleDetector.Find(TestBlock) == nullptr, TEXT("A circular dependency was found in scratch buffer."));
-				CycleDetector.Add(TestBlock);
-				TestBlock = TestBlock->NextBlock;
-			}
-		}
-#endif // TEDS_SCRATCHBUFFER_ENABLE_ADDITIONAL_TESTS
-	}
-}
-
-FTypedElementDatabaseScratchBuffer::FBlockControllerMap& FTypedElementDatabaseScratchBuffer::GetThreadLocalBlockControllerMap()
-{
-	thread_local static FBlockControllerMap LocalBlockMap;
-	return LocalBlockMap;
-}
-
-FTypedElementDatabaseScratchBuffer::FBlockController& FTypedElementDatabaseScratchBuffer::GetThreadLocalBlockController()
-{
-	return GetThreadLocalBlockControllerMap().FindOrAddControllerFor(*this);
-}
-
-void FTypedElementDatabaseScratchBuffer::ConfigureDestructorTail(
-	FDestructorTail& Destructor, DestructorFunction Callback, void* Object, int32 Count)
-{
-	GetThreadLocalBlockController().ConfigureDestructorTail(Destructor, Callback, Object, Count);
-}
-
-
-//
-// FBlockController
-//
-
-FTypedElementDatabaseScratchBuffer::FBlockController::FBlockController(const TWeakPtr<FTypedElementDatabaseScratchBuffer>& InParent)
-	: Parent(InParent)
-{
-	Block = GetEmptyBlock();
+	return CurrentAllocator.load()->Malloc(Size, Alignment);
 }
 
-FTypedElementDatabaseScratchBuffer::FBlockController::~FBlockController()
+void FTypedElementDatabaseScratchBuffer::BatchDelete()
 {
-	checkf(Block != nullptr, TEXT("FBlockController is expected to always have a valid block to work with."));
-	RecycleBlock();
-}
-
-void* FTypedElementDatabaseScratchBuffer::FBlockController::Allocate(size_t Size, size_t Alignment, uint64 LocalFrameId)
-{
-	checkf(Alignment <= alignof(FBlock), TEXT("Alignment of %i for allocation in database scratch buffer exceeds maximal alignment of %i."),
-		static_cast<int>(Alignment), static_cast<int>(alignof(FBlock)));
-
-	if (Size <= MaxAllocationSize())
-	{
-		size_t Index = Align(Block->Front, Alignment);
-		size_t NewFront = Index + Size;
-
-		if (NewFront > FBlock::BlockSize)
-		{
-			// No more space left in the block so recycle the current one and create a new one.
-			RecycleBlock();
-			Block = GetEmptyBlock();
-			Index = 0;
-			NewFront = Size;
-		}
-
-		Block->LastTouchedByFrame = LocalFrameId;
-		Block->Front = NewFront;
-		return Block->Buffer + Index;
-	}
-	else
-	{
-		struct FExtendedBufferStore
-		{
-			void* ExtendedBuffer;
-			FDestructorTail Destructor;
-		};
-
-		// Allocate memory directly as the requested size doesn't fit in a block.
-		void* ExtendedBuffer = FMemory::Malloc(Size, Alignment);
-		
-		// Add an entry to the block with the sole purpose of deleting the extended buffer.
-		void* BufferStoreAddress = Allocate(sizeof(FExtendedBufferStore), alignof(FExtendedBufferStore), LocalFrameId);
-		FExtendedBufferStore* BufferStore = reinterpret_cast<FExtendedBufferStore*>(BufferStoreAddress);
-		BufferStore->ExtendedBuffer = ExtendedBuffer;
-
-		ConfigureDestructorTail(BufferStore->Destructor,
-			[](void* Object, [[maybe_unused]] int32 ObjectCount)
-			{
-				FMemory::Free(reinterpret_cast<FExtendedBufferStore*>(Object)->ExtendedBuffer);
-			}, BufferStore);
-
-		return ExtendedBuffer;
-	}
-}
-
-FTypedElementDatabaseScratchBuffer::FBlock* FTypedElementDatabaseScratchBuffer::FBlockController::GetEmptyBlock()
-{
-	if (TSharedPtr<FTypedElementDatabaseScratchBuffer> ParentInstance = Parent.Pin())
-	{
-		FBlock* FrontBlock = ParentInstance->AvailableBlocks.load();
-		// Because the full and available queues are separate, this does not suffer from the ABA problem.
-		while (FrontBlock != nullptr && !ParentInstance->AvailableBlocks.compare_exchange_weak(FrontBlock, FrontBlock->NextBlock)){};
-		if (FrontBlock)
-		{
-			FrontBlock->NextBlock = nullptr;
-			return FrontBlock;
-		}
-	}
-
-	// If null then there are no more buffers left so create a new one.
-	return new FBlock();
-}
-
-void FTypedElementDatabaseScratchBuffer::FBlockController::RecycleBlock()
-{
-	if (TSharedPtr<FTypedElementDatabaseScratchBuffer> ParentInstance = Parent.Pin())
-	{
-		Block->Front = 0;
-		
-		// Replace the front of the linked list with full blocks with the latest full block.
-		// Because the full and available queues are separate, this does not suffer from the ABA problem.
-		FBlock* FrontBlock = ParentInstance->FullBlocks;
-		do
-		{
-			Block->NextBlock = FrontBlock;
-		} while (!ParentInstance->FullBlocks.compare_exchange_strong(FrontBlock, Block));
-	}
-	else
-	{
-		// The assumption is that if the scratch buffer has been deleted none of the data in the
-		// current buffer will be used so it can be deleted. It also means that the scratch buffer
-		// is not going to recycle this block for this thread so needs to be deleted here to avoid
-		// a memory leak.
-		delete Block;
-	}
-	Block = nullptr;
-}
-
-void FTypedElementDatabaseScratchBuffer::FBlockController::ConfigureDestructorTail(
-	FDestructorTail& Destructor, DestructorFunction Callback, void* Object, int32 Count)
-{
-	Destructor.Destructor = Callback;
-	Destructor.PreviousTail = Block->DestructionTail;
-	Destructor.StructOffset = static_cast<int32>(reinterpret_cast<char*>(&Destructor) - reinterpret_cast<char*>(Object));
-	Destructor.InstanceCount = Count;
-	Block->DestructionTail = &Destructor;
-}
-
-
-
-//
-// FBlockControllerMap
-//
-
-FTypedElementDatabaseScratchBuffer::FBlockController& 
-	FTypedElementDatabaseScratchBuffer::FBlockControllerMap::FindOrAddControllerFor(FTypedElementDatabaseScratchBuffer& ScratchBuffer)
-{
-	TWeakPtr<FTypedElementDatabaseScratchBuffer>* ParentsIt = Parents.GetData();
-	const uint32 Count = Parents.Num();
-	for (uint32 Index = 0; Index < Count; ++Index, ++ParentsIt)
-	{
-		if (ParentsIt->HasSameObject(&ScratchBuffer))
-		{
-			return Controllers[Index];
-		}
-	}
-
-	Parents.Add(ScratchBuffer.AsWeak());
-	uint32 NewIndex = Controllers.Emplace(ScratchBuffer.AsWeak());
-	return Controllers[NewIndex];
+	MemoryAllocator* Current = CurrentAllocator.exchange(LeastRecentAllocator);
+	LeastRecentAllocator = PreviousAllocator;
+	PreviousAllocator = Current;
+	LeastRecentAllocator->BulkDelete();
 }
diff --git a/Engine/Plugins/Experimental/TypedElementsDataStorage/Source/TypedElementsDataStorage/Private/TypedElementDatabaseScratchBuffer.h b/Engine/Plugins/Experimental/TypedElementsDataStorage/Source/TypedElementsDataStorage/Private/TypedElementDatabaseScratchBuffer.h
--- a/Engine/Plugins/Experimental/TypedElementsDataStorage/Source/TypedElementsDataStorage/Private/TypedElementDatabaseScratchBuffer.h
+++ b/Engine/Plugins/Experimental/TypedElementsDataStorage/Source/TypedElementsDataStorage/Private/TypedElementDatabaseScratchBuffer.h
@@ -3,25 +3,27 @@
 #pragma once
 
 #include <atomic>
-#include <type_traits>
-#include "GenericPlatform/GenericPlatformMath.h"
-#include "Math/UnrealMath.h"
-#include "Misc/MTAccessDetector.h"
-#include "Templates/AlignmentTemplates.h"
-#include "Templates/SharedPointer.h"
-#include "Templates/UnrealTemplate.h"
+#include "Experimental/ConcurrentLinearAllocator.h"
 
-// A thread-safe memory allocator that uses linear allocation using a chain of recycled memory blocks.
-// This provides a fast and lightweight way to allocate temporary memory for intermediate values that will
-// live at best until the end of the frame. Objects that require destruction will be deleted once a block is
-// recycled. A block is only recycled when it's full, so during low activity on a particular thread it can
-// take significantly longer than a single frame before an object is destroyed.
-class FTypedElementDatabaseScratchBuffer : public TSharedFromThis<FTypedElementDatabaseScratchBuffer>
+// A thread-safe memory allocator that uses linear allocation. This provides a fast and lightweight way to allocate temporary memory 
+// for intermediate values that will live at best until the end of the frame.
+class FTypedElementDatabaseScratchBuffer
 {
-	friend class FBlockController;
-	friend class FBlockControllerMap;
+private:
+	struct FScratchBufferAllocationTag
+	{
+		static constexpr uint32 BlockSize = 64 * 1024;		// Blocksize used to allocate from
+		static constexpr bool AllowOversizedBlocks = true;  // The allocator supports oversized Blocks and will store them in a separate Block with counter 1
+		static constexpr bool RequiresAccurateSize = true;  // GetAllocationSize returning the accurate size of the allocation otherwise it could be relaxed to return the size to the end of the Block
+		static constexpr bool InlineBlockAllocation = false;  // Inline or Noinline the BlockAllocation which can have an impact on Performance
+		static constexpr const char* TagName = "TEDS_ScratchBuffer";
+
+		using Allocator = TBlockAllocationLockFreeCache<BlockSize, FOsAllocator>;
+	};
+	using MemoryAllocator = TConcurrentLinearBulkObjectAllocator<FScratchBufferAllocationTag>;
+
 public:
-	~FTypedElementDatabaseScratchBuffer();
+	FTypedElementDatabaseScratchBuffer();
 
 	void* Allocate(size_t Size, size_t Alignment);
 	template<typename T, typename... ArgTypes>
@@ -29,154 +31,37 @@ public:
 	template<typename T, typename... ArgTypes>
 	T* EmplaceArray(int32 Count, const ArgTypes&... Args);
 
-	// Increments the frame id by one. The frame id is used to track which blocks were used in the current frame and
-	// can't therefore be safely removed. Call this at the start of a frame.
-	void NextFrame();
-	// Makes previously filled up blocks available again. This can only safely be called when Allocate(...) can not
-	// be called. As allocations should only be called when Mass is running processors, it means this function
-	// can be safely called when the processors are not running.
-	void RecycleBlocks();
-
-	constexpr static int32 MaxAllocationSize();
+	// Activates a new allocator and deletes all commands and objects in the least recently touched scratch buffer.
+	void BatchDelete();
+constexpr static int32 MaxAllocationSize();
 
 private:
-	using DestructorFunction = void (*)(void* Object, int32 ObjectCount);
-
-	// Track objects that need to be deleted. This will happen when a chain of blocks is recycled. This can
-	// mean that an objects stays around for several frames if a block isn't heavily used.
-	struct FDestructorTail
-	{
-		FDestructorTail* PreviousTail;
-		DestructorFunction Destructor;
-		int32 StructOffset;
-		int32 InstanceCount;
-	};
-
-	// Use the largest SIMD alignment for current hardware just in case. The buffer is sufficiently large
-	// that the overhead won't matter.
-	struct alignas(64) FBlock
-	{
-		constexpr static size_t BlockSize = 64 * 1024; // 64kb blocks
-		
-		char Buffer[BlockSize]; // Keep at top for alignment.
-		// Indicates the last frame this block was used in
-		uint64 LastTouchedByFrame = 0;
-		FDestructorTail* DestructionTail = nullptr;
-		std::atomic<FBlock*> NextBlock = nullptr;
-		// Offset into the buffer where the next allocation starts.
-		uint32 Front = 0;
-	};
-
-	// Simple wrapper around the buffer. This is to guarantee that when the thread gets destroyed the buffer is cleaned up.
-	struct FBlockController
-	{
-		explicit FBlockController(const TWeakPtr<FTypedElementDatabaseScratchBuffer>& InParent);
-		~FBlockController();
-
-		void* Allocate(size_t Size, size_t Alignment, uint64 LocalFrameId);
-		FBlock* GetEmptyBlock();
-		void RecycleBlock();
-		void ConfigureDestructorTail(FDestructorTail& Destructor, DestructorFunction Callback, void* Object, int32 Count = 1);
-
-		TWeakPtr<FTypedElementDatabaseScratchBuffer> Parent;
-		FBlock* Block;
-	};
-
-	// Simple container to hold multiple block controllers and return the one appropriate for the thread and scratch buffer.
-	struct FBlockControllerMap
-	{
-		TArray<TWeakPtr<FTypedElementDatabaseScratchBuffer>> Parents;
-		TArray<FBlockController> Controllers;
-
-		FBlockController& FindOrAddControllerFor(FTypedElementDatabaseScratchBuffer& ScratchBuffer);
-	};
-
-	// Returns the block controller map unique to the thread that calls this function.
-	FBlockControllerMap& GetThreadLocalBlockControllerMap();
-	// Returns the block controller that's unique to the thread that calls this function.
-	FBlockController& GetThreadLocalBlockController();
-	void ConfigureDestructorTail(FDestructorTail& Destructor, DestructorFunction Callback, void* Object, int32 Count = 1);
-
-	// Detect any invalid calls to recycle blocks while allocating memory.
-	UE_MT_DECLARE_RW_ACCESS_DETECTOR(AccessDetector);
-
-	std::atomic<FBlock*> AvailableBlocks = nullptr;
-	std::atomic<FBlock*> FullBlocks = nullptr;
-	
-	// The number that identifies the current frame. Used to keep track of the blocks that were used in this frame or earlier.
-	uint64 FrameId = 0;
+	// Using a triple buffered approach because the direct API in TEDS (those calls that can be made directly to the API and don't go
+	// through a context) are not required to be atomic. As such it's possible that data for a command is stored in allocator A while
+	// the command is in allocator B if those calls are issued while TEDS is closing it's processing cycle. With double buffering this
+	// would result in allocator A being flushed thus clearing out the data for the command. Using a triple buffered approach will
+	// cause the clearing to be delayed by a frame, avoid this problem. This however does assume that all data and command issuing happens
+	// within a single tick, though for the direct API this should always be true.
+
+	MemoryAllocator Allocators[3];
+	std::atomic<MemoryAllocator*> CurrentAllocator;
+	MemoryAllocator* PreviousAllocator;
+	MemoryAllocator* LeastRecentAllocator;
 };
 
-
-//
-// Implementations
-//
-
 template<typename T, typename... ArgTypes>
 T* FTypedElementDatabaseScratchBuffer::Emplace(ArgTypes... Args)
 {
-	if constexpr (std::is_trivially_destructible_v<T>)
-	{
-		return new(Allocate(sizeof(T), alignof(T))) T(Forward<ArgTypes>(Args)...);
-	}
-	else
-	{
-		struct Wrapper
-		{
-			T Instance;
-			FDestructorTail Destructor;
-		};
-		Wrapper* Result = reinterpret_cast<Wrapper*>(Allocate(sizeof(Wrapper), alignof(Wrapper)));
-		new(&Result->Instance) T(Forward<ArgTypes>(Args)...);
-
-		ConfigureDestructorTail(Result->Destructor, 
-			[](void* Object, [[maybe_unused]] int32 ObjectCount)
-			{ 
-				reinterpret_cast<T*>(Object)->~T(); 
-			}, &Result->Instance);
-
-		return &Result->Instance;
-	}
+	return CurrentAllocator.load()->Create<T>(Forward<ArgTypes>(Args)...);
 }
 
 template<typename T, typename... ArgTypes>
 T* FTypedElementDatabaseScratchBuffer::EmplaceArray(int32 Count, const ArgTypes&... Args)
 {
-	T* Result;	
-	if constexpr (std::is_trivially_destructible_v<T>)
-	{
-		Result = reinterpret_cast<T*>(Allocate(sizeof(T) * Count, alignof(T)));
-	}
-	else
-	{
-		constexpr size_t Alignment = FMath::Max(alignof(T), alignof(FDestructorTail));
-		int32 DestructorTailOffset = Align(sizeof(T) * Count, alignof(FDestructorTail));
-		int32 Size = DestructorTailOffset + sizeof(FDestructorTail);
-		Result = reinterpret_cast<T*>(Allocate(Size, Alignment));
-
-		FDestructorTail* Destructor = reinterpret_cast<FDestructorTail*>(reinterpret_cast<char*>(Result) + DestructorTailOffset);
-		ConfigureDestructorTail(*Destructor,
-			[](void* Object, int32 ObjectCount)
-			{
-				T* Array = reinterpret_cast<T*>(Object);
-				for (int32 Count = 0; Count < ObjectCount; ++Count)
-				{
-					Array->~T();
-					++Array;
-				}
-			}, Result, Count);
-	}
-
-	for (int32 Index = 0; Index < Count; ++Index)
-	{
-		new(Result + Index) T(Args...);
-	}
-
-	return Result;
+	return CurrentAllocator.load()->CreateArray<T>(Count, Args...);
 }
 
 constexpr int32 FTypedElementDatabaseScratchBuffer::MaxAllocationSize()
 {
-	// Leave room for at least one destructor tail in case a single large object is added that requires destruction.
-	return FBlock::BlockSize - sizeof(FDestructorTail);
+	return FScratchBufferAllocationTag::BlockSize;
 }
diff --git a/Engine/Plugins/Experimental/TypedElementsDataStorage/Source/TypedElementsDataStorage/Tests/TypedElementDatabaseScratchBufferTests.cpp b/Engine/Plugins/Experimental/TypedElementsDataStorage/Source/TypedElementsDataStorage/Tests/TypedElementDatabaseScratchBufferTests.cpp
--- a/Engine/Plugins/Experimental/TypedElementsDataStorage/Source/TypedElementsDataStorage/Tests/TypedElementDatabaseScratchBufferTests.cpp
+++ b/Engine/Plugins/Experimental/TypedElementsDataStorage/Source/TypedElementsDataStorage/Tests/TypedElementDatabaseScratchBufferTests.cpp
@@ -8,7 +8,7 @@
 #include "TypedElementDatabaseScratchBuffer.h"
 #include "Tests/TestHarnessAdapter.h"
 
-// Tries to shutdown all threads within the maximum wait time. If it takes longer than the minimum wait time the test will fail.
+// Tries to shutdown all threads within the maximum wait time. If it takes longer than the maximum wait time the test will fail.
 template<int ThreadCount, class Rep, class Period>
 static void WaitForFinalization(std::condition_variable& WaitVariable, std::mutex& WaitMutex, 
 	std::atomic<bool>& bRunFlag, std::atomic<uint32>& CompletedThreadCount, const std::chrono::duration<Rep, Period>& MaxWaitTime)
@@ -24,17 +24,6 @@ static void WaitForFinalization(std::condition_variable& WaitVariable, std::mute
 	REQUIRE(CompletedThreadCount == ThreadCount);
 }
 
-// Spin locks for the maximum wait time to let threads resume
-template<int ThreadCount, class Rep, class Period>
-static void WaitForPause(std::atomic<uint32>& PausedCounter, const std::chrono::duration<Rep, Period>& MaxWaitTime)
-{
-	std::chrono::system_clock::time_point StartTime = std::chrono::system_clock::now();
-	while (PausedCounter != ThreadCount
-		&& std::chrono::duration_cast<std::chrono::milliseconds>(std::chrono::system_clock::now() - StartTime) < MaxWaitTime)
-	{}
-	REQUIRE(PausedCounter == ThreadCount);
-}
-
 TEST_CASE("TypedElementsDataStorage::Scratch Buffer - MT (FTypedElementDatabaseScratchBuffer)", "[ApplicationContextMask][EngineFilter]")
 {
 	SECTION("Stress test without recycling.")
@@ -42,7 +31,7 @@ TEST_CASE("TypedElementsDataStorage::Scratch Buffer - MT (FTypedElementDatabaseS
 		using namespace std::chrono_literals;
 		constexpr uint32 ThreadCount = 8;
 
-		TSharedPtr<FTypedElementDatabaseScratchBuffer> Buffer = MakeShared<FTypedElementDatabaseScratchBuffer>();
+		FTypedElementDatabaseScratchBuffer Buffer;
 
 		std::atomic<bool> bKeepRunning = true;
 		std::atomic<uint32> FailedAllocations = 0;
@@ -57,7 +46,7 @@ TEST_CASE("TypedElementsDataStorage::Scratch Buffer - MT (FTypedElementDatabaseS
 				{
 					while (bKeepRunning)
 					{
-						void* Data = Buffer->Allocate(128, 4);
+						void* Data = Buffer.Allocate(128, 4);
 						if (Data == nullptr)
 						{
 							FailedAllocations++;
@@ -87,16 +76,11 @@ TEST_CASE("TypedElementsDataStorage::Scratch Buffer - MT (FTypedElementDatabaseS
 
 		constexpr uint32 ThreadCount = 8;
 
-		TSharedPtr<FTypedElementDatabaseScratchBuffer> Buffer = MakeShared<FTypedElementDatabaseScratchBuffer>();
+		FTypedElementDatabaseScratchBuffer Buffer;
 
 		std::atomic<bool> bKeepRunning = true;
-		std::atomic<bool> bPause = false;
-		std::atomic<uint32> PausedCounter = 0;
 		std::atomic<uint32> FailedAllocations = 0;
 
-		std::condition_variable PauseSync;
-		std::mutex PauseMutex;
-
 		std::atomic<uint32> CompletedThreadCount = 0;
 		std::condition_variable CompletionSync;
 		std::mutex CompletionMutex;
@@ -106,8 +90,7 @@ TEST_CASE("TypedElementsDataStorage::Scratch Buffer - MT (FTypedElementDatabaseS
 		for (uint32 ThreadCounter = 0; ThreadCounter < ThreadCount; ++ThreadCounter)
 		{
 			Threads.Emplace([
-				ThreadCounter, &Buffer, &bKeepRunning, &FailedAllocations, &bPause, &PausedCounter, 
-				&PauseSync, &PauseMutex, &CompletedThreadCount, &CompletionSync]()
+				ThreadCounter, &Buffer, &bKeepRunning, &FailedAllocations, &CompletedThreadCount, &CompletionSync]()
 				{
 					static constexpr uint32 MemoryAllocationSizes[] = { 128, 32, 14332, 741, 8871, 48, 27335 };
 					static constexpr uint32 MemoryAllocationSizeCount = sizeof(MemoryAllocationSizes) / sizeof(MemoryAllocationSizes[0]);
@@ -115,26 +98,20 @@ TEST_CASE("TypedElementsDataStorage::Scratch Buffer - MT (FTypedElementDatabaseS
 					uint32 MemoryAllocationSizeIndex = ThreadCounter;
 					while (bKeepRunning)
 					{
-						// Not using % as there's currently a false positive in MSVC that marks it warning (C6385).
+						// Not using % as there's currently a false positive in MSVC that marks it as a warning (C6385).
 						++MemoryAllocationSizeIndex;
 						if (MemoryAllocationSizeIndex >= MemoryAllocationSizeCount)
 						{
 							MemoryAllocationSizeIndex = 0;
 						}
 
-						void* Data = Buffer->Allocate(static_cast<size_t>(MemoryAllocationSizes[MemoryAllocationSizeIndex]), 4);
+						void* Data = Buffer.Allocate(static_cast<size_t>(MemoryAllocationSizes[MemoryAllocationSizeIndex]), 4);
 						if (Data == nullptr)
 						{
 							FailedAllocations++;
 						}
-
-						if (bPause)
-						{
-							++PausedCounter; // Announce that this thread is paused.
-							std::unique_lock<std::mutex> WaitLock(PauseMutex); // Wait to be re-enabled again.
-							PauseSync.wait(WaitLock);
-							++PausedCounter; // Announce that this thread has resumed.
-						}
+						// Sleep a little while as to not to end up with an excessive amount of memory allocations.
+						std::this_thread::sleep_for(1ms);
 					}
 					if (++CompletedThreadCount == ThreadCount)
 					{
@@ -143,26 +120,10 @@ TEST_CASE("TypedElementsDataStorage::Scratch Buffer - MT (FTypedElementDatabaseS
 				});
 		}
 
-		// While the threads are continuously allocating memory, use the main thread to periodically pause the other threads 
-		// and recycle blocks.
+		// While the threads are continuously allocating memory, use the main thread to periodically recycle blocks.
 		for (uint32 IterationCounter = 0; IterationCounter < 60; ++IterationCounter)
 		{
-			PausedCounter = 0;
-			bPause = true;
-			WaitForPause<ThreadCount>(PausedCounter, 500ms);
-			
-			Buffer->NextFrame();
-			Buffer->RecycleBlocks();
-
-			// Resume all threads.
-			{
-				std::unique_lock<std::mutex> WaitLock(PauseMutex);
-				PausedCounter = 0;
-				bPause = false;
-				PauseSync.notify_all();
-			}
-			WaitForPause<ThreadCount>(PausedCounter, 500ms);
-		
+			Buffer.BatchDelete();
 			std::this_thread::sleep_for(33ms);
 		}
 		WaitForFinalization<ThreadCount>(CompletionSync, CompletionMutex, bKeepRunning, CompletedThreadCount, 500ms);
@@ -185,68 +146,55 @@ TEST_CASE("TypedElementsDataStorage::Scratch Buffer (FTypedElementDatabaseScratc
 
 	SECTION("Allocate small block")
 	{
-		TSharedPtr<FTypedElementDatabaseScratchBuffer> Buffer = MakeShared<FTypedElementDatabaseScratchBuffer>();
-		void* Data = Buffer->Allocate(4, 4);
+		FTypedElementDatabaseScratchBuffer Buffer;
+		void* Data = Buffer.Allocate(4, 4);
 		CHECK(Data != nullptr);
 	}
 
 	SECTION("Allocate over-sized block")
 	{
-		TSharedPtr<FTypedElementDatabaseScratchBuffer> Buffer = MakeShared<FTypedElementDatabaseScratchBuffer>();
-		void* Data = Buffer->Allocate(Buffer->MaxAllocationSize() * 4, 4);
+		FTypedElementDatabaseScratchBuffer Buffer;
+		void* Data = Buffer.Allocate(Buffer.MaxAllocationSize() * 4, 4);
 		CHECK(Data != nullptr);
 	}
 
-	SECTION("Allocate exactly remaining block size")
-	{
-		TSharedPtr<FTypedElementDatabaseScratchBuffer> Buffer = MakeShared<FTypedElementDatabaseScratchBuffer>();
-		void* Data1 = Buffer->Allocate(1, 1);
-		void* Data2 = Buffer->Allocate(Buffer->MaxAllocationSize() - 1, 1);
-
-		REQUIRE(Data1 != nullptr);
-		REQUIRE(Data2 != nullptr);
-		// Check if they're both in the same buffer.
-		CHECK(reinterpret_cast<uintptr_t>(Data2) == reinterpret_cast<uintptr_t>(Data1) + 1);
-	}
-
 	SECTION("Alignment respected")
 	{
-		TSharedPtr<FTypedElementDatabaseScratchBuffer> Buffer = MakeShared<FTypedElementDatabaseScratchBuffer>();
-		void* Data1 = Buffer->Allocate(1, 1);
-		void* Data2 = Buffer->Allocate(4, 4);
+		FTypedElementDatabaseScratchBuffer Buffer;
+		void* Data1 = Buffer.Allocate(1, 1);
+		void* Data2 = Buffer.Allocate(4, 4);
 
 		REQUIRE(Data1 != nullptr);
 		REQUIRE(Data2 != nullptr);
-		// Check if they're both in the same buffer.
+		// Check alignment is respected.
 		CHECK((reinterpret_cast<uintptr_t>(Data2) & 3) == 0);
 	}
 
 	SECTION("Multiple blocks used.")
 	{
-		TSharedPtr<FTypedElementDatabaseScratchBuffer> Buffer = MakeShared<FTypedElementDatabaseScratchBuffer>();
+		FTypedElementDatabaseScratchBuffer Buffer;
 
-		int32 IncrementCount = (Buffer->MaxAllocationSize() * 4 /* Fill 4 blocks*/) / 64 /* With 64 by allocations */;
+		int32 IncrementCount = (Buffer.MaxAllocationSize() * 4 /* Fill 4 blocks*/) / 64 /* With 64 byte allocations */;
 		for (int32 Counter = 0; Counter < IncrementCount; ++Counter)
 		{
-			void* Data = Buffer->Allocate(64, 4);
+			void* Data = Buffer.Allocate(64, 4);
 			REQUIRE(Data != nullptr);
 		}
 	}
 
 	SECTION("Recycle full blocks.")
 	{
-		TSharedPtr<FTypedElementDatabaseScratchBuffer> Buffer = MakeShared<FTypedElementDatabaseScratchBuffer>();
+		FTypedElementDatabaseScratchBuffer Buffer;
 
 		for (int32 Iterations = 0; Iterations < 16; ++Iterations)
 		{
-			int32 IncrementCount = Buffer->MaxAllocationSize() / 64;
-			for (int32 Counter = 0; Counter <= IncrementCount; ++Counter) // Increment one more that needed to fill the block and make sure it's recycled.
+			int32 IncrementCount = Buffer.MaxAllocationSize() / 64;
+			for (int32 Counter = 0; Counter <= IncrementCount; ++Counter) // Increment one more than needed to fill the block and make sure it's recycled.
 			{
-				void* Data = Buffer->Allocate(64, 4);
+				void* Data = Buffer.Allocate(64, 4);
 				REQUIRE(Data != nullptr);
 			}
-			Buffer->NextFrame();
-			Buffer->RecycleBlocks();
+			Buffer.BatchDelete();
 		}
 	}
 }
diff --git a/Engine/Source/Runtime/TypedElementFramework/Public/Elements/Interfaces/TypedElementQueryStorageInterfaces.h b/Engine/Source/Runtime/TypedElementFramework/Public/Elements/Interfaces/TypedElementQueryStorageInterfaces.h
--- a/Engine/Source/Runtime/TypedElementFramework/Public/Elements/Interfaces/TypedElementQueryStorageInterfaces.h
+++ b/Engine/Source/Runtime/TypedElementFramework/Public/Elements/Interfaces/TypedElementQueryStorageInterfaces.h
@@ -89,7 +89,7 @@ namespace TypedElementDataStorage
 	 */
 	struct IQueryContext : public ICommonQueryContext
 	{
-		using ObjectMoveOperator = void (*)(void* Destination, void* Source);
+		using ObjectCopyOrMove = void (*)(const UScriptStruct* TypeInfo, void* Destination, void* Source);
 
 		virtual ~IQueryContext() = default;
 
@@ -118,25 +118,27 @@ namespace TypedElementDataStorage
 		 * group.
 		 */
 		virtual void RemoveRows(TConstArrayView<RowHandle> Rows) = 0;
+		
 		/**
-		 * Add a new column of the provided type if one does not exist.
-		 * Returns a staged column which is used to copy into the database at a later time via the UStructScript Copy operator
-		 * It is the caller's responsibility to ensure the staged column's constructor is called.  The caller
-		 * may modify other properties of the column.
-		 *
-		 * Note: The addition and modification of the column will not be done immediately done. Instead it will be deferred until the end of the tick group.
+		 * Adds the provided column to the requested row.
+		 * 
+		 * Note: The addition of the column will not be immediately done. Instead it will be deferred until the end of the tick group. Changes
+		 * made to the return column will still be applied when the column is added to the row.
 		 */
-		virtual void* AddColumnUninitialized(RowHandle Row, const UScriptStruct* ColumnType) = 0;
-		
+		template<typename ColumnType>
+		ColumnType& AddColumn(RowHandle Row, ColumnType&& Column);
+		/**
+		 * Adds new empty columns to a row of the provided type. The addition will not be immediately done but delayed until the end of the
+		 * tick group.
+		 */
+		template<typename... Columns>
+		void AddColumns(RowHandle Row);
 		/**
-		 * Add a new column of the provided type if one does not exist.
-		 * Returns a staged column which is used to copy into the database at a later time via the provided Move operator
-		 * It is the caller's responsibility to ensure the staged column's constructor is called.  The caller
-		 * may modify other properties of the column.
-		 *
-		 * Note: The addition and modification of the column will not be done immediately done. Instead it will be deferred until the end of the tick group.
+		 * Adds new empty columns to the listed rows of the provided type. The addition will not be immediately done but delayed until the end of the
+		 * tick group.
 		 */
-		virtual void* AddColumnUninitialized(RowHandle Row, const UScriptStruct* ObjectType, ObjectMoveOperator Mover) = 0;
+		template<typename... Columns>
+		void AddColumns(TConstArrayView<RowHandle> Rows);
 		/**
 		 * Adds new empty columns to a row of the provided type. The addition will not be immediately done but delayed until the end of the
 		 * tick group.
@@ -147,7 +149,37 @@ namespace TypedElementDataStorage
 		 * tick group.
 		 */
 		virtual void AddColumns(TConstArrayView<RowHandle> Rows, TConstArrayView<const UScriptStruct*> ColumnTypes) = 0;
+		/**
+		 * Add a new uninitialized column of the provided type if one does not exist.
+		 * Returns a staged column which is used to copy into the database at a later time via the UStructScript Copy operator at the end
+		 * of the tick group.
+		 * It is the caller's responsibility to ensure the staged column's constructor is called. The caller may modify other 
+		 * values in the column.
+		 * This function can not be used to add a tag as tags do not contain any data.
+		 */
+		virtual void* AddColumnUninitialized(RowHandle Row, const UScriptStruct* ColumnType) = 0;
+		/**
+		 * Add a new uninitialized column of the provided type if one does not exist.
+		 * Returns a staged column which is used to copy/move into the database at a later time via the provided relocator at the end of
+		 * the tick group.
+		 * It is the caller's responsibility to ensure the staged column's constructor is called. The caller may modify other 
+		 * values in the column.
+		 * This function can not be used to add a tag as tags do not contain any data.
+		 */
+		virtual void* AddColumnUninitialized(RowHandle Row, const UScriptStruct* ObjectType, ObjectCopyOrMove Relocator) = 0;
 
+		/**
+		 * Removes columns of the provided types from a row. The removal will not be immediately done but delayed until the end of the
+		 * tick group.
+		 */
+		template<typename... Columns>
+		void RemoveColumns(RowHandle Row);
+		/**
+		 * Removes columns of the provided types from the listed rows. The removal will not be immediately done but delayed until the end of the
+		 * tick group.
+		 */
+		template<typename... Columns>
+		void RemoveColumns(TConstArrayView<RowHandle> Rows);
 		/**
 		 * Removes columns of the provided types from a row. The removal will not be immediately done but delayed until the end of the
 		 * tick group.
@@ -185,21 +217,6 @@ namespace TypedElementDataStorage
 		 * are executed as part of their parent query and are not scheduled separately.
 		 */
 		virtual FQueryResult RunSubquery(int32 SubqueryIndex, RowHandle Row, SubqueryCallbackRef Callback) = 0;
-
-
-
-		// Utility functions
-
-		template<typename ColumnType>
-		ColumnType& AddColumn(RowHandle Row, ColumnType&& Column);
-		template<typename... Columns>
-		void AddColumns(RowHandle Row);
-		template<typename... Columns>
-		void AddColumns(TConstArrayView<RowHandle> Rows);
-		template<typename... Columns>
-		void RemoveColumns(RowHandle Row);
-		template<typename... Columns>
-		void RemoveColumns(TConstArrayView<RowHandle> Rows);
 	};
 } // namespace TypedElementDataStorage
 
@@ -238,9 +255,9 @@ namespace TypedElementDataStorage
 		if constexpr (std::is_move_constructible_v<ColumnType>)
 		{
 			void* Address = AddColumnUninitialized(Row, TypeInfo,
-				[](void* Destination, void* Source)
+				[](const UScriptStruct*, void* Destination, void* Source)
 				{
-					new(Destination) ColumnType(MoveTemp(*reinterpret_cast<ColumnType*>(Source)));
+					*reinterpret_cast<ColumnType*>(Destination) = MoveTemp(*reinterpret_cast<ColumnType*>(Source));
 				});
 			return *(new(Address) ColumnType(Forward<ColumnType>(Column)));
 		}
