Fixed up more bool-taking calls to take EAllowShrinking instead.

diff --git a/Engine/Plugins/Cameras/GameplayCameras/Source/GameplayCameras/Private/Core/BlendStackCameraNode.cpp b/Engine/Plugins/Cameras/GameplayCameras/Source/GameplayCameras/Private/Core/BlendStackCameraNode.cpp
--- a/Engine/Plugins/Cameras/GameplayCameras/Source/GameplayCameras/Private/Core/BlendStackCameraNode.cpp
+++ b/Engine/Plugins/Cameras/GameplayCameras/Source/GameplayCameras/Private/Core/BlendStackCameraNode.cpp
@@ -175,7 +175,7 @@ void UBlendStackCameraNode::OnRun(const FCameraNodeRunParams& Params, FCameraNod
 	// Pop out camera modes that have been blended out.
 	if (bAutoPop && PopEntriesBelow != INDEX_NONE)
 	{
-		Entries.RemoveAt(0, PopEntriesBelow, false);
+		Entries.RemoveAt(0, PopEntriesBelow, EAllowShrinking::No);
 	}
 
 	// Reset first frame flags.
diff --git a/Engine/Plugins/Experimental/PCG/Source/PCG/Private/Graph/PCGGraphExecutor.cpp b/Engine/Plugins/Experimental/PCG/Source/PCG/Private/Graph/PCGGraphExecutor.cpp
--- a/Engine/Plugins/Experimental/PCG/Source/PCG/Private/Graph/PCGGraphExecutor.cpp
+++ b/Engine/Plugins/Experimental/PCG/Source/PCG/Private/Graph/PCGGraphExecutor.cpp
@@ -1380,7 +1380,7 @@ void FPCGGraphExecutor::CullInactiveDownstreamNodes(FPCGTaskId InCompletedTaskId
 
 	while (!PinIdsToDeactivate.IsEmpty())
 	{
-		const FPCGPinId PinId = PinIdsToDeactivate.Pop(/*bAllowShrinking=*/false);
+		const FPCGPinId PinId = PinIdsToDeactivate.Pop(EAllowShrinking::No);
 		const FPCGTaskId PinTaskId = PCGPinIdHelpers::GetNodeIdFromPinId(PinId);
 
 		PCGGraphExecutionLogging::LogTaskCullingBeginLoop(PinTaskId, PCGPinIdHelpers::GetPinIndexFromPinId(PinId), PinIdsToDeactivate);
diff --git a/Engine/Plugins/Experimental/TypedElementsDataStorage/Source/TypedElementsDataStorage/Private/TypedElementDatabaseCompatibility.cpp b/Engine/Plugins/Experimental/TypedElementsDataStorage/Source/TypedElementsDataStorage/Private/TypedElementDatabaseCompatibility.cpp
--- a/Engine/Plugins/Experimental/TypedElementsDataStorage/Source/TypedElementsDataStorage/Private/TypedElementDatabaseCompatibility.cpp
+++ b/Engine/Plugins/Experimental/TypedElementsDataStorage/Source/TypedElementsDataStorage/Private/TypedElementDatabaseCompatibility.cpp
@@ -652,8 +652,7 @@ void UTypedElementDatabaseCompatibility::TickObjectSync()
 					}
 				}
 				const int32 RowHandleCount = RowHandleIndex;
-				const bool bAllowShrinking = false;
-				RowHandles.SetNum(RowHandleCount, bAllowShrinking);
+				RowHandles.SetNum(RowHandleCount, EAllowShrinking::No);
 			}
 
 			ObjectsNeedingFullSync.Reset();
