Populate TEv1 handles in TEDS

* Initial submission will populate only Actor TypedElementHandles
* Disabled by default - enabled with `TEDS.TypedElementBridge.Enable 1`

diff --git a/Engine/Plugins/Experimental/TypedElementsDataStorage/Source/TypedElementsDataStorage/Private/Compatibility/TEDSTypedElementActorHandleFactory.cpp b/Engine/Plugins/Experimental/TypedElementsDataStorage/Source/TypedElementsDataStorage/Private/Compatibility/TEDSTypedElementActorHandleFactory.cpp
new file mode 100644
--- /dev/null
+++ b/Engine/Plugins/Experimental/TypedElementsDataStorage/Source/TypedElementsDataStorage/Private/Compatibility/TEDSTypedElementActorHandleFactory.cpp
@@ -0,0 +1,127 @@
+﻿// Copyright Epic Games, Inc. All Rights Reserved.
+
+#include "TEDSTypedElementActorHandleFactory.h"
+
+#include "MassActorSubsystem.h"
+#include "Compatibility/TEDSTypedElementBridge.h"
+#include "Elements/Columns/TypedElementCompatibilityColumns.h"
+#include "Elements/Framework/EngineElementsLibrary.h"
+#include "Elements/Framework/TypedElementList.h"
+#include "Elements/Framework/TypedElementQueryBuilder.h"
+#include "Elements/Framework/TypedElementRegistry.h"
+
+void UTEDSTypedElementActorHandleFactory::PreRegister(ITypedElementDataStorageInterface& DataStorage)
+{
+	Super::PreRegister(DataStorage);
+	
+	BridgeEnableDelegateHandle = UTEDSTypedElementBridge::OnEnabled().AddUObject(this, &UTEDSTypedElementActorHandleFactory::HandleBridgeEnabled);
+}
+
+void UTEDSTypedElementActorHandleFactory::PreShutdown(ITypedElementDataStorageInterface& DataStorage)
+{
+	UTEDSTypedElementBridge::OnEnabled().Remove(BridgeEnableDelegateHandle);
+	BridgeEnableDelegateHandle.Reset();
+}
+
+void UTEDSTypedElementActorHandleFactory::RegisterQueries(ITypedElementDataStorageInterface& DataStorage)
+{
+	Super::RegisterQueries(DataStorage);
+
+	if (UTEDSTypedElementBridge::IsEnabled())
+	{
+		RegisterQuery_ActorHandlePopulate(DataStorage);
+	}
+
+	using namespace TypedElementQueryBuilder;
+	using namespace TypedElementDataStorage;
+	GetAllActorsQuery = DataStorage.RegisterQuery(
+	Select()
+		.ReadOnly<FMassActorFragment>()
+	.Compile());
+}
+
+void UTEDSTypedElementActorHandleFactory::RegisterQuery_ActorHandlePopulate(ITypedElementDataStorageInterface& DataStorage)
+{
+	using namespace TypedElementQueryBuilder;
+	using namespace TypedElementDataStorage;
+	using DSI = ITypedElementDataStorageInterface;
+	if (!ensureMsgf(ActorHandlePopulateQuery == TypedElementInvalidQueryHandle, TEXT("Already registered query")))
+	{
+		return;
+	}
+	
+	ActorHandlePopulateQuery = DataStorage.RegisterQuery(
+	Select(TEXT("Populate actor typed element handles"),
+	FProcessor(DSI::EQueryTickPhase::FrameEnd, DataStorage.GetQueryTickGroupName(DSI::EQueryTickGroups::SyncWidgets))
+		.ForceToGameThread(true),
+		[](DSI::IQueryContext& Context, TypedElementRowHandle Row, const FTypedElementUObjectColumn& ObjectColumn)
+		{
+			if (UObject* Object = ObjectColumn.Object.Get())
+			{
+				checkSlow(Cast<AActor>(Object));
+				FTypedElementHandle Handle = UEngineElementsLibrary::AcquireEditorActorElementHandle(static_cast<AActor*>(Object));
+				Context.AddColumn(Row, FTEDSTypedElementColumn
+				{
+					.Handle = Handle
+				});
+			}
+
+			Context.RemoveColumns<FPendingPopulateTypedElementColumn>(Row);
+		}
+	)
+	.Where()
+		.All<FPendingPopulateTypedElementColumn, FMassActorFragment>()
+	.Compile());
+}
+
+void UTEDSTypedElementActorHandleFactory::HandleBridgeEnabled(bool bEnabled)
+{
+	ITypedElementDataStorageInterface* DataStorage = UTypedElementRegistry::GetInstance()->GetMutableDataStorage();
+	
+	if (bEnabled)
+	{
+		using namespace TypedElementQueryBuilder;
+		using namespace TypedElementDataStorage;
+		using DSI = ITypedElementDataStorageInterface;
+		
+		// Populate all the rows
+		TArray<RowHandle> CollatedRowHandles;
+		TArray<TWeakObjectPtr<const AActor>> Actors;
+		
+		DataStorage->RunQuery(GetAllActorsQuery, CreateDirectQueryCallbackBinding([&CollatedRowHandles, &Actors](IDirectQueryContext& Context, const FMassActorFragment* Fragments)
+		{
+			TConstArrayView<RowHandle> RowHandles = Context.GetRowHandles();
+
+			CollatedRowHandles.Append(RowHandles);
+
+			TConstArrayView<const FMassActorFragment> FragmentView(Fragments, Context.GetRowCount());
+
+			Actors.Reserve(Actors.Num() + FragmentView.Num());
+			for (const FMassActorFragment& Fragment : FragmentView)
+			{
+				const AActor* Actor = Fragment.Get(FMassActorFragment::EActorAccess::OnlyWhenAlive);
+				Actors.Add(Actor);
+			}
+		}));
+
+		for (int32 Index = 0, End = CollatedRowHandles.Num(); Index < End; ++Index)
+		{
+			if (const AActor* Actor = Actors[Index].Get())
+			{
+				FTypedElementHandle Handle = UEngineElementsLibrary::AcquireEditorActorElementHandle(Actor);
+				DataStorage->AddOrGetColumn(CollatedRowHandles[Index], FTEDSTypedElementColumn
+				{
+					.Handle = Handle
+				});
+			}
+		}
+
+		RegisterQuery_ActorHandlePopulate(*DataStorage);
+	}
+	else
+	{
+		DataStorage->UnregisterQuery(ActorHandlePopulateQuery);
+		ActorHandlePopulateQuery = TypedElementInvalidQueryHandle;
+	}
+}
+
diff --git a/Engine/Plugins/Experimental/TypedElementsDataStorage/Source/TypedElementsDataStorage/Private/Compatibility/TEDSTypedElementActorHandleFactory.h b/Engine/Plugins/Experimental/TypedElementsDataStorage/Source/TypedElementsDataStorage/Private/Compatibility/TEDSTypedElementActorHandleFactory.h
new file mode 100644
--- /dev/null
+++ b/Engine/Plugins/Experimental/TypedElementsDataStorage/Source/TypedElementsDataStorage/Private/Compatibility/TEDSTypedElementActorHandleFactory.h
@@ -0,0 +1,36 @@
+﻿// Copyright Epic Games, Inc. All Rights Reserved.
+
+#pragma once
+
+#include "Elements/Framework/TypedElementHandle.h"
+#include "Elements/Interfaces/TypedElementDataStorageInterface.h"
+#include "Elements/Interfaces/TypedElementDataStorageFactory.h"
+#include "UObject/UObjectGlobals.h"
+
+#include "TEDSTypedElementActorHandleFactory.generated.h"
+
+class UTypedElementRegistry;
+
+/**
+ * This class is responsible for acquiring and registering Actor TypedElementHandles
+ * with TEDS/
+ */
+UCLASS(Transient)
+class UTEDSTypedElementActorHandleFactory : public UTypedElementDataStorageFactory
+{
+	GENERATED_BODY()
+public:
+
+	virtual void PreRegister(ITypedElementDataStorageInterface& DataStorage) override;
+	virtual void RegisterQueries(ITypedElementDataStorageInterface& DataStorage) override;
+	virtual void PreShutdown(ITypedElementDataStorageInterface& DataStorage) override;
+private:
+	virtual void RegisterQuery_ActorHandlePopulate(ITypedElementDataStorageInterface& DataStorage);
+	void HandleBridgeEnabled(bool bEnabled);
+
+	FDelegateHandle BridgeEnableDelegateHandle;
+	TypedElementQueryHandle ActorHandlePopulateQuery = TypedElementInvalidQueryHandle;
+	TypedElementQueryHandle GetAllActorsQuery = TypedElementInvalidQueryHandle;
+};
+
+
diff --git a/Engine/Plugins/Experimental/TypedElementsDataStorage/Source/TypedElementsDataStorage/Private/Compatibility/TEDSTypedElementBridge.cpp b/Engine/Plugins/Experimental/TypedElementsDataStorage/Source/TypedElementsDataStorage/Private/Compatibility/TEDSTypedElementBridge.cpp
new file mode 100644
--- /dev/null
+++ b/Engine/Plugins/Experimental/TypedElementsDataStorage/Source/TypedElementsDataStorage/Private/Compatibility/TEDSTypedElementBridge.cpp
@@ -0,0 +1,157 @@
+﻿// Copyright Epic Games, Inc. All Rights Reserved.
+
+#include "Compatibility/TEDSTypedElementBridge.h"
+
+#include "MassActorSubsystem.h"
+#include "Elements/Columns/TypedElementCompatibilityColumns.h"
+#include "Elements/Framework/EngineElementsLibrary.h"
+#include "Elements/Framework/TypedElementList.h"
+#include "Elements/Framework/TypedElementQueryBuilder.h"
+#include "Elements/Framework/TypedElementRegistry.h"
+
+namespace UTEDSTypedElementBridge_Private
+{
+	bool bBridgeEnabled = false;
+	TYPEDELEMENTSDATASTORAGE_API FAutoConsoleVariableRef CVarBridgeEnabled(
+		TEXT("TEDS.TyepedElementBridge.Enable"),
+		bBridgeEnabled,
+		TEXT("Automatically populated TEDS with TypedElementHandles"));
+}
+
+uint8 UTEDSTypedElementBridge::GetOrder() const
+{
+	return 110;
+}
+
+void UTEDSTypedElementBridge::PreRegister(ITypedElementDataStorageInterface& DataStorage)
+{
+	Super::PreRegister(DataStorage);
+	DebugEnabledDelegateHandle = UTEDSTypedElementBridge_Private::CVarBridgeEnabled->OnChangedDelegate().AddUObject(this, &UTEDSTypedElementBridge::HandleOnEnabled);
+}
+
+void UTEDSTypedElementBridge::PreShutdown(ITypedElementDataStorageInterface& DataStorage)
+{
+	UTEDSTypedElementBridge_Private::CVarBridgeEnabled->OnChangedDelegate().Remove(DebugEnabledDelegateHandle);
+	DebugEnabledDelegateHandle.Reset();
+	CleanupTypedElementColumns(DataStorage);
+
+	Super::PreShutdown(DataStorage);
+}
+
+void UTEDSTypedElementBridge::RegisterQueries(ITypedElementDataStorageInterface& DataStorage)
+{
+	Super::RegisterQueries(DataStorage);
+
+	if (IsEnabled())
+	{
+		RegisterQuery_NewUObject(DataStorage);
+	}
+}
+
+void UTEDSTypedElementBridge::RegisterQuery_NewUObject(ITypedElementDataStorageInterface& DataStorage)
+{
+	using namespace TypedElementQueryBuilder;
+	using DSI = ITypedElementDataStorageInterface;
+
+	if (!ensureMsgf(QueryHandle == TypedElementInvalidQueryHandle, TEXT("Query already registered")))
+	{
+		return;
+	}
+
+	// TODO Need a two-step initialize since the FObserver::Add gets called before FTypedElementUObjectColumn is
+	// ever populated.
+
+	QueryHandle = DataStorage.RegisterQuery(
+		Select(
+		TEXT("Watch for new UObjects"),
+			FObserver::OnAdd<FTypedElementUObjectColumn>(),
+			[](DSI::IQueryContext& Context, TypedElementRowHandle Row, const FTypedElementUObjectColumn& Object)
+			{
+				Context.AddColumns<FPendingPopulateTypedElementColumn>(Row);
+			}
+		)
+		.Compile());
+	
+	RemoveTypedElementRowHandleQuery = DataStorage.RegisterQuery(
+		Select()
+			.ReadOnly<FTEDSTypedElementColumn>()
+		.Compile());
+	
+	RemovePendingTypedElementRowHandleQuery = DataStorage.RegisterQuery(
+			Select().Where().All<FPendingPopulateTypedElementColumn>()
+		.Compile());
+}
+
+void UTEDSTypedElementBridge::UnregisterQuery_NewUObject(ITypedElementDataStorageInterface& DataStorage)
+{
+	DataStorage.UnregisterQuery(QueryHandle);
+	QueryHandle = TypedElementInvalidQueryHandle;
+}
+
+FOnTEDSTypedElementBridgeEnable& UTEDSTypedElementBridge::OnEnabled()
+{
+	static FOnTEDSTypedElementBridgeEnable OnBridgeEnabled;
+	return OnBridgeEnabled;
+}
+
+bool UTEDSTypedElementBridge::IsEnabled()
+{
+	return UTEDSTypedElementBridge_Private::CVarBridgeEnabled->GetBool();
+}
+
+void UTEDSTypedElementBridge::CleanupTypedElementColumns(ITypedElementDataStorageInterface& DataStorage)
+{
+	using namespace TypedElementQueryBuilder;
+	using namespace TypedElementDataStorage;
+
+	// Remove any TEv1 handles
+	{
+		TArray<TypedElementRowHandle> Handles;
+		using namespace TypedElementQueryBuilder;
+		DataStorage.RunQuery(
+			RemoveTypedElementRowHandleQuery,
+			CreateDirectQueryCallbackBinding(
+				[&Handles](ITypedElementDataStorageInterface::IDirectQueryContext& Context)
+				{
+					Handles.Append(Context.GetRowHandles());
+				}));
+		
+		DataStorage.BatchAddRemoveColumns(TConstArrayView<TypedElementRowHandle>(Handles), {}, {FTEDSTypedElementColumn::StaticStruct()});
+	}
+
+	// Remove any columns on row indicating they are pending population of their TEv1 columnn
+	{
+		TArray<TypedElementRowHandle> Handles;
+		using namespace TypedElementQueryBuilder;
+		DataStorage.RunQuery(
+			RemoveTypedElementRowHandleQuery,
+			CreateDirectQueryCallbackBinding(
+				[&Handles](ITypedElementDataStorageInterface::IDirectQueryContext& Context)
+				{
+					Handles.Append(Context.GetRowHandles());
+				}));
+		
+		DataStorage.BatchAddRemoveColumns(TConstArrayView<TypedElementRowHandle>(Handles), {}, {FPendingPopulateTypedElementColumn::StaticStruct()});
+	}
+}
+
+void UTEDSTypedElementBridge::HandleOnEnabled(IConsoleVariable* CVar)
+{
+	ITypedElementDataStorageInterface* DataStorage = UTypedElementRegistry::GetInstance()->GetMutableDataStorage();
+	bool bIsEnabled = CVar->GetBool();
+
+	if (bIsEnabled)
+	{
+		DataStorage->UnregisterQuery(QueryHandle);
+		QueryHandle = TypedElementInvalidQueryHandle;
+		
+		RegisterQuery_NewUObject(*DataStorage);
+		UTEDSTypedElementBridge::OnEnabled().Broadcast(bIsEnabled);
+	}
+	else
+	{
+		UTEDSTypedElementBridge::OnEnabled().Broadcast(bIsEnabled);
+		UnregisterQuery_NewUObject(*DataStorage);
+		CleanupTypedElementColumns(*DataStorage);
+	}
+}
diff --git a/Engine/Plugins/Experimental/TypedElementsDataStorage/Source/TypedElementsDataStorage/Public/Compatibility/TEDSTypedElementBridge.h b/Engine/Plugins/Experimental/TypedElementsDataStorage/Source/TypedElementsDataStorage/Public/Compatibility/TEDSTypedElementBridge.h
new file mode 100644
--- /dev/null
+++ b/Engine/Plugins/Experimental/TypedElementsDataStorage/Source/TypedElementsDataStorage/Public/Compatibility/TEDSTypedElementBridge.h
@@ -0,0 +1,63 @@
+﻿// Copyright Epic Games, Inc. All Rights Reserved.
+
+#pragma once
+
+#include "Elements/Framework/TypedElementHandle.h"
+#include "Elements/Interfaces/TypedElementDataStorageInterface.h"
+#include "Elements/Interfaces/TypedElementDataStorageFactory.h"
+#include "UObject/UObjectGlobals.h"
+
+#include "TEDSTypedElementBridge.generated.h"
+
+class UTypedElementRegistry;
+
+DECLARE_MULTICAST_DELEGATE_OneParam(FOnTEDSTypedElementBridgeEnable, bool /*bEnabled*/);
+
+/**
+ * This class is responsible for running queries that will ensure TypedElementHandles
+ * are cleaned up when TEDS is shut down.
+ */
+UCLASS(Transient)
+class UTEDSTypedElementBridge : public UTypedElementDataStorageFactory
+{
+	GENERATED_BODY()
+public:
+	~UTEDSTypedElementBridge() override = default;
+	virtual uint8 GetOrder() const override;
+	virtual void PreRegister(ITypedElementDataStorageInterface& DataStorage) override;
+	virtual void PreShutdown(ITypedElementDataStorageInterface& DataStorage) override;
+	virtual void RegisterQueries(ITypedElementDataStorageInterface& DataStorage) override;
+
+	static TYPEDELEMENTSDATASTORAGE_API FOnTEDSTypedElementBridgeEnable& OnEnabled();
+	static TYPEDELEMENTSDATASTORAGE_API bool IsEnabled();
+
+private:
+	void RegisterQuery_NewUObject(ITypedElementDataStorageInterface& DataStorage);
+	void UnregisterQuery_NewUObject(ITypedElementDataStorageInterface& DataStorage);
+	void CleanupTypedElementColumns(ITypedElementDataStorageInterface& DataStorage);
+	void HandleOnEnabled(IConsoleVariable* CVar);
+	
+	TypedElementQueryHandle QueryHandle = TypedElementDataStorage::InvalidQueryHandle;
+	TypedElementQueryHandle RemoveTypedElementRowHandleQuery = TypedElementDataStorage::InvalidQueryHandle;
+	TypedElementQueryHandle RemovePendingTypedElementRowHandleQuery = TypedElementDataStorage::InvalidQueryHandle;
+	FDelegateHandle DebugEnabledDelegateHandle;
+};
+
+// A column which contains a TypedElementHandle
+// TypedElement, in this context, refers to the interface-based
+// TypedElements instead of the data-based TypedElementDataStorage (TEDS)
+USTRUCT()
+struct FTEDSTypedElementColumn : public FTypedElementDataStorageColumn
+{
+	GENERATED_BODY()
+	
+	FTypedElementHandle Handle;
+};
+
+// Used by processors to defer assignment of TEDSTypedElementColumn
+// until processor runs
+USTRUCT()
+struct FPendingPopulateTypedElementColumn : public FTypedElementDataStorageTag
+{
+	GENERATED_BODY()
+};
\ No newline at end of file
