Skein / ViewportStatus: Update TTypedElementRevisionControlFactory query.

The query now assigns an outline color using the criteria described in the ProductSpec.
- Yellow, if a newer package is available.
- Red, if someone else locked the package,
- Green, if package is added locally.
- Blue, if package is locked by user.

The color indices still, as before, rely on the container in UEditorStyleSettings.

I've introduced a FSCCNotCurrentTag to know in the processor if something is at the latest revision as it will not be obvious for some source control systems to provide an ordering based on just FSCCRevisionId (if it's a GUID for example).

#rnx

diff --git a/Engine/Plugins/Developer/PerforceSourceControl/Source/PerforceSourceControl/Private/PerforceSourceControlOperations.cpp b/Engine/Plugins/Developer/PerforceSourceControl/Source/PerforceSourceControl/Private/PerforceSourceControlOperations.cpp
--- a/Engine/Plugins/Developer/PerforceSourceControl/Source/PerforceSourceControl/Private/PerforceSourceControlOperations.cpp
+++ b/Engine/Plugins/Developer/PerforceSourceControl/Source/PerforceSourceControl/Private/PerforceSourceControlOperations.cpp
@@ -2100,6 +2100,7 @@ bool FPerforceUpdateStatusWorker::UpdateStates() const
 		SyncTagFromState(bIsCheckedOutByOther, FSCCExternallyEditedTag::StaticStruct());
 		SyncTagFromState(State.IsCheckedOut(), FSCCLockedTag::StaticStruct());
 		SyncTagFromState(State.Changelist.IsInitialized(), FSCCInChangelistTag::StaticStruct());
+		SyncTagFromState(State.LocalRevNumber < State.DepotRevNumber, FSCCNotCurrentTag::StaticStruct());
 		DataStorage->AddColumns(Row, ToAdd);
 		// Not using batched removal for now because it ensures that all are present
 		for (UScriptStruct* Column : ToRemove)
diff --git a/Engine/Plugins/Experimental/TypedElementsDataStorage/Source/TypedElementsDataStorageRevisionControl/Source/Private/RevisionControlProcessors.cpp b/Engine/Plugins/Experimental/TypedElementsDataStorage/Source/TypedElementsDataStorageRevisionControl/Source/Private/RevisionControlProcessors.cpp
--- a/Engine/Plugins/Experimental/TypedElementsDataStorage/Source/TypedElementsDataStorageRevisionControl/Source/Private/RevisionControlProcessors.cpp
+++ b/Engine/Plugins/Experimental/TypedElementsDataStorage/Source/TypedElementsDataStorageRevisionControl/Source/Private/RevisionControlProcessors.cpp
@@ -18,67 +18,57 @@
 
 extern TYPEDELEMENTSDATASTORAGE_API FAutoConsoleVariableRef CVarAutoPopulateState;
 
-FAutoConsoleCommandWithArgsAndOutputDevice SetSelectionSCCStateConsoleCommand(
-	TEXT("TEDS.Debug.SetSCCState"),
-	TEXT("Adds a source control state to selected objects."),
-	FConsoleCommandWithArgsAndOutputDeviceDelegate::CreateLambda([](const TArray<FString>& Args, FOutputDevice& Output)
+static bool GEnableOutlines = false;
+static FAutoConsoleVariableRef CVarEnableOutlines(
+	TEXT("TEDS.RevisionControl.UseOutlines"),
+	GEnableOutlines,
+	TEXT("Use revision control status outlines in the viewport")
+);
+
+static uint8 DetermineOutlineColorIndex(const TypedElementDataStorage::ICommonQueryContext& Context)
+{
+	if (GEnableOutlines)
+	{
+		// TODO: Get this information from TEDS instead of hardcoded pointing into UEditorStyleSettings::AdditionalSelectionColors?
+		constexpr uint8 BasicSelectionColorCount = 2;
+		constexpr uint8 IndexBlue = BasicSelectionColorCount + 0;
+		constexpr uint8 IndexPurple = BasicSelectionColorCount + 1;
+		constexpr uint8 IndexPink = BasicSelectionColorCount + 2;
+		constexpr uint8 IndexRed = BasicSelectionColorCount + 3;
+		constexpr uint8 IndexYellow = BasicSelectionColorCount + 4;
+		constexpr uint8 IndexGreen = BasicSelectionColorCount + 5;
+
+		// Check if the package is outdated because there is a newer version available.
+		if (Context.HasColumn<FSCCNotCurrentTag>())
 		{
-			using namespace TypedElementQueryBuilder;
-			using DSI = ITypedElementDataStorageInterface;
+			return IndexYellow;
+		}
 
-			TRACE_CPUPROFILER_EVENT_SCOPE(TEDS.Debug.SetSCCState);
+		// Check if the package is locked by someone else.
+		if (Context.HasColumn<FSCCExternallyLockedColumn>())
+		{
+			return IndexRed;
+		}
 
-			if (ITypedElementDataStorageInterface* DataStorage = UTypedElementRegistry::GetInstance()->GetMutableDataStorage())
+		// Check if the package is added locally.
+		if (Context.HasColumn<FSCCStatusColumn>())
+		{
+			const FSCCStatusColumn* StatusColumn = Context.GetColumn<FSCCStatusColumn>();
+			if (StatusColumn->Modification == ESCCModification::Added)
 			{
-				static TypedElementQueryHandle AllSelectedQuery = TypedElementInvalidQueryHandle;
-				if (AllSelectedQuery == TypedElementInvalidQueryHandle)
-				{
-					AllSelectedQuery = DataStorage->RegisterQuery(
-						Select()
-							.ReadOnly<FTypedElementPackageReference>()
-						.Where()
-							.All<FTypedElementSelectionColumn>()
-						.Compile());
-				}
-				
-				if (AllSelectedQuery == TypedElementInvalidQueryHandle)
-				{
-					return;
-				}
-
-				ESCCModification Modification = ESCCModification::Modified;
-
-				if (Args.Num() > 0)
-				{
-					// Parse the index
-					int32 StateIndex;
-					LexFromString(StateIndex, *Args[0]);
-
-					if (!(StateIndex >= 0 && StateIndex <= int32(ESCCModification::Conflicted)))
-					{
-						Output.Log(TEXT("State index out of range"));
-						return;
-					}
-
-					Modification = static_cast<ESCCModification>(StateIndex);
-				}
-				
-				TArray<TypedElementRowHandle> PackageRowHandles;
-				
-				DataStorage->RunQuery(AllSelectedQuery, CreateDirectQueryCallbackBinding(
-					[&PackageRowHandles] (DSI::IDirectQueryContext& Context, const FTypedElementPackageReference& PackageReference)
-					{
-						PackageRowHandles.Add(PackageReference.Row);
-					})
-				);
-
-				for (TypedElementRowHandle Row : PackageRowHandles)
-				{
-					DataStorage->AddOrGetColumn<FSCCStatusColumn>(Row)->Modification = Modification;
-				}
+				return IndexGreen;
 			}
 		}
-	));
+
+		// Check if the package is locked by self.
+		if (Context.HasColumn<FSCCLockedTag>())
+		{
+			return IndexBlue;
+		}
+	}
+
+	return 0; // Default outline color.
+}
 
 void UTypedElementRevisionControlFactory::RegisterTables(ITypedElementDataStorageInterface& DataStorage)
 {
@@ -96,21 +86,20 @@ void UTypedElementRevisionControlFactory::RegisterQueries(ITypedElementDataStora
 
 	TypedElementQueryHandle ObjectToSCCQuery = DataStorage.RegisterQuery(
 		Select()
-			.ReadOnly<FSCCStatusColumn>()
+			.ReadOnly<FTypedElementPackagePathColumn>()
 		.Compile());
 
 	DataStorage.RegisterQuery(
 		Select(
 			TEXT("Change selection outline colors based on SCC status"),
 			// This is in PrePhysics because the outline->actor query is in DuringPhysics and contexts don't flush changes between tick groups
-			FProcessor(DSI::EQueryTickPhase::PrePhysics, DataStorage.GetQueryTickGroupName(DSI::EQueryTickGroups::SyncExternalToDataStorage)),			
+			FProcessor(DSI::EQueryTickPhase::PrePhysics, DataStorage.GetQueryTickGroupName(DSI::EQueryTickGroups::SyncExternalToDataStorage)),
 			[](DSI::IQueryContext& Context, TypedElementRowHandle ObjectRow, const FTypedElementPackageReference& PackageReference)
 			{
 				Context.RunSubquery(0, PackageReference.Row, CreateSubqueryCallbackBinding(
-					[&Context, &ObjectRow](DSI::ISubqueryContext& SubQueryContext, const FSCCStatusColumn& Status)
+					[&Context, &ObjectRow](DSI::ISubqueryContext& SubQueryContext)
 					{
-						constexpr int BasicSelectionColorCount = 2;
-						Context.AddColumn<FTypedElementViewportColorColumn>(ObjectRow, { .SelectionOutlineColorIndex = static_cast<uint8>(int(Status.Modification) + BasicSelectionColorCount) });
+						Context.AddColumn<FTypedElementViewportColorColumn>(ObjectRow, { .SelectionOutlineColorIndex = DetermineOutlineColorIndex(SubQueryContext) });
 						Context.AddColumns<FTypedElementSyncBackToWorldTag>(ObjectRow);
 					})
 				);
diff --git a/Engine/Source/Runtime/TypedElementFramework/Public/Elements/Columns/TypedElementRevisionControlColumns.h b/Engine/Source/Runtime/TypedElementFramework/Public/Elements/Columns/TypedElementRevisionControlColumns.h
--- a/Engine/Source/Runtime/TypedElementFramework/Public/Elements/Columns/TypedElementRevisionControlColumns.h
+++ b/Engine/Source/Runtime/TypedElementFramework/Public/Elements/Columns/TypedElementRevisionControlColumns.h
@@ -30,6 +30,12 @@ struct FSCCExternallyEditedTag final : public FTypedElementDataStorageTag
 	GENERATED_BODY()
 };
 
+USTRUCT(meta = (DisplayName = "Not at the latest revision"))
+struct FSCCNotCurrentTag final : public FTypedElementDataStorageTag
+{
+	GENERATED_BODY()
+};
+
 UENUM()
 enum class ESCCModification
 {
