Fixed "General.Header" not being registered.

The TEDS UI purpose for general headers was never registered. This forced otherwise general headers to be specific to the outliner.

diff --git a/Engine/Plugins/Experimental/TypedElementsDataStorage/Source/TypedElementsDataStorageUI/Private/TypedElementGeneralRegistration.cpp b/Engine/Plugins/Experimental/TypedElementsDataStorage/Source/TypedElementsDataStorageUI/Private/TypedElementGeneralRegistration.cpp
--- a/Engine/Plugins/Experimental/TypedElementsDataStorage/Source/TypedElementsDataStorageUI/Private/TypedElementGeneralRegistration.cpp
+++ b/Engine/Plugins/Experimental/TypedElementsDataStorage/Source/TypedElementsDataStorageUI/Private/TypedElementGeneralRegistration.cpp
@@ -7,6 +7,7 @@
 #define LOCTEXT_NAMESPACE "TypedElementsUI_GeneralRegistration"
 
 const FName UTypedElementGeneralRegistrationFactory::CellPurpose(TEXT("General.Cell"));
+const FName UTypedElementGeneralRegistrationFactory::HeaderPurpose(TEXT("General.Header"));
 const FName UTypedElementGeneralRegistrationFactory::CellDefaultPurpose(TEXT("General.Cell.Default"));
 const FName UTypedElementGeneralRegistrationFactory::HeaderDefaultPurpose(TEXT("General.Header.Default"));
 
@@ -14,7 +15,10 @@ void UTypedElementGeneralRegistrationFactory::RegisterWidgetPurposes(ITypedEleme
 {
 	DataStorageUi.RegisterWidgetPurpose(CellPurpose, ITypedElementDataStorageUiInterface::EPurposeType::UniqueByNameAndColumn,
 		LOCTEXT("CellPurpose", "General purpose widgets that can be used as cells for specific columns or column combinations."));
-	
+
+	DataStorageUi.RegisterWidgetPurpose(HeaderPurpose, ITypedElementDataStorageUiInterface::EPurposeType::UniqueByNameAndColumn,
+		LOCTEXT("HeaderDefaultPurpose", "General purpose widget that can be used as a header."));
+
 	DataStorageUi.RegisterWidgetPurpose(HeaderDefaultPurpose, ITypedElementDataStorageUiInterface::EPurposeType::UniqueByName,
 		LOCTEXT("HeaderDefaultPurpose", "The default widget to use in headers if no other specialization is provided."));
 	
diff --git a/Engine/Plugins/Experimental/TypedElementsDataStorage/Source/TypedElementsDataStorageUI/Public/TypedElementGeneralRegistration.h b/Engine/Plugins/Experimental/TypedElementsDataStorage/Source/TypedElementsDataStorageUI/Public/TypedElementGeneralRegistration.h
--- a/Engine/Plugins/Experimental/TypedElementsDataStorage/Source/TypedElementsDataStorageUI/Public/TypedElementGeneralRegistration.h
+++ b/Engine/Plugins/Experimental/TypedElementsDataStorage/Source/TypedElementsDataStorageUI/Public/TypedElementGeneralRegistration.h
@@ -14,6 +14,7 @@ class TYPEDELEMENTSDATASTORAGEUI_API UTypedElementGeneralRegistrationFactory : p
 
 public:
 	static const FName CellPurpose;
+	static const FName HeaderPurpose;
 	static const FName CellDefaultPurpose;
 	static const FName HeaderDefaultPurpose;
 
