Resolve name conflict in localization.

diff --git a/Engine/Plugins/Experimental/TypedElementsDataStorage/Source/TypedElementsDataStorageUI/Private/TypedElementGeneralRegistration.cpp b/Engine/Plugins/Experimental/TypedElementsDataStorage/Source/TypedElementsDataStorageUI/Private/TypedElementGeneralRegistration.cpp
--- a/Engine/Plugins/Experimental/TypedElementsDataStorage/Source/TypedElementsDataStorageUI/Private/TypedElementGeneralRegistration.cpp
+++ b/Engine/Plugins/Experimental/TypedElementsDataStorage/Source/TypedElementsDataStorageUI/Private/TypedElementGeneralRegistration.cpp
@@ -14,16 +14,16 @@ const FName UTypedElementGeneralRegistrationFactory::HeaderDefaultPurpose(TEXT("
 void UTypedElementGeneralRegistrationFactory::RegisterWidgetPurposes(ITypedElementDataStorageUiInterface& DataStorageUi) const
 {
 	DataStorageUi.RegisterWidgetPurpose(CellPurpose, ITypedElementDataStorageUiInterface::EPurposeType::UniqueByNameAndColumn,
-		LOCTEXT("CellPurpose", "General purpose widgets that can be used as cells for specific columns or column combinations."));
+		LOCTEXT("GeneralCellPurpose", "General purpose widgets that can be used as cells for specific columns or column combinations."));
 
 	DataStorageUi.RegisterWidgetPurpose(HeaderPurpose, ITypedElementDataStorageUiInterface::EPurposeType::UniqueByNameAndColumn,
-		LOCTEXT("HeaderDefaultPurpose", "General purpose widget that can be used as a header."));
+		LOCTEXT("GeneralHeaderPurpose", "General purpose widget that can be used as a header."));
 
 	DataStorageUi.RegisterWidgetPurpose(HeaderDefaultPurpose, ITypedElementDataStorageUiInterface::EPurposeType::UniqueByName,
-		LOCTEXT("HeaderDefaultPurpose", "The default widget to use in headers if no other specialization is provided."));
+		LOCTEXT("GeneralHeaderDefaultPurpose", "The default widget to use in headers if no other specialization is provided."));
 	
 	DataStorageUi.RegisterWidgetPurpose(CellDefaultPurpose, ITypedElementDataStorageUiInterface::EPurposeType::UniqueByName,
-		LOCTEXT("CellDefaultPurpose", "The default widget to use in cells if no other specialization is provided."));
+		LOCTEXT("GeneralCellDefaultPurpose", "The default widget to use in cells if no other specialization is provided."));
 }
 
 #undef LOCTEXT_NAMESPACE
