Remove the use of a "populate tag" for TEDS TypedElement bridge
* Fix typo in CVar

diff --git a/Engine/Plugins/Experimental/TypedElementsDataStorage/Source/TypedElementsDataStorage/Private/Compatibility/TEDSTypedElementActorHandleFactory.cpp b/Engine/Plugins/Experimental/TypedElementsDataStorage/Source/TypedElementsDataStorage/Private/Compatibility/TEDSTypedElementActorHandleFactory.cpp
--- a/Engine/Plugins/Experimental/TypedElementsDataStorage/Source/TypedElementsDataStorage/Private/Compatibility/TEDSTypedElementActorHandleFactory.cpp
+++ b/Engine/Plugins/Experimental/TypedElementsDataStorage/Source/TypedElementsDataStorage/Private/Compatibility/TEDSTypedElementActorHandleFactory.cpp
@@ -52,8 +52,7 @@ void UTEDSTypedElementActorHandleFactory::RegisterQuery_ActorHandlePopulate(ITyp
 	
 	ActorHandlePopulateQuery = DataStorage.RegisterQuery(
 	Select(TEXT("Populate actor typed element handles"),
-	FProcessor(DSI::EQueryTickPhase::FrameEnd, DataStorage.GetQueryTickGroupName(DSI::EQueryTickGroups::SyncWidgets))
-		.ForceToGameThread(true),
+		FObserver::OnAdd<FTypedElementUObjectColumn>(),
 		[](DSI::IQueryContext& Context, TypedElementRowHandle Row, const FTypedElementUObjectColumn& ObjectColumn)
 		{
 			if (UObject* Object = ObjectColumn.Object.Get())
@@ -65,12 +64,10 @@ void UTEDSTypedElementActorHandleFactory::RegisterQuery_ActorHandlePopulate(ITyp
 					.Handle = Handle
 				});
 			}
-
-			Context.RemoveColumns<FPendingPopulateTypedElementColumn>(Row);
 		}
 	)
 	.Where()
-		.All<FPendingPopulateTypedElementColumn, FMassActorFragment>()
+		.All<FMassActorFragment>()
 	.Compile());
 }
 
diff --git a/Engine/Plugins/Experimental/TypedElementsDataStorage/Source/TypedElementsDataStorage/Private/Compatibility/TEDSTypedElementBridge.cpp b/Engine/Plugins/Experimental/TypedElementsDataStorage/Source/TypedElementsDataStorage/Private/Compatibility/TEDSTypedElementBridge.cpp
--- a/Engine/Plugins/Experimental/TypedElementsDataStorage/Source/TypedElementsDataStorage/Private/Compatibility/TEDSTypedElementBridge.cpp
+++ b/Engine/Plugins/Experimental/TypedElementsDataStorage/Source/TypedElementsDataStorage/Private/Compatibility/TEDSTypedElementBridge.cpp
@@ -13,7 +13,7 @@ namespace UTEDSTypedElementBridge_Private
 {
 	bool bBridgeEnabled = false;
 	TYPEDELEMENTSDATASTORAGE_API FAutoConsoleVariableRef CVarBridgeEnabled(
-		TEXT("TEDS.TyepedElementBridge.Enable"),
+		TEXT("TEDS.TypedElementBridge.Enable"),
 		bBridgeEnabled,
 		TEXT("Automatically populated TEDS with TypedElementHandles"));
 }
@@ -52,40 +52,15 @@ void UTEDSTypedElementBridge::RegisterQuery_NewUObject(ITypedElementDataStorageI
 {
 	using namespace TypedElementQueryBuilder;
 	using DSI = ITypedElementDataStorageInterface;
-
-	if (!ensureMsgf(QueryHandle == TypedElementInvalidQueryHandle, TEXT("Query already registered")))
-	{
-		return;
-	}
-
-	// TODO Need a two-step initialize since the FObserver::Add gets called before FTypedElementUObjectColumn is
-	// ever populated.
-
-	QueryHandle = DataStorage.RegisterQuery(
-		Select(
-		TEXT("Watch for new UObjects"),
-			FObserver::OnAdd<FTypedElementUObjectColumn>(),
-			[](DSI::IQueryContext& Context, TypedElementRowHandle Row, const FTypedElementUObjectColumn& Object)
-			{
-				Context.AddColumns<FPendingPopulateTypedElementColumn>(Row);
-			}
-		)
-		.Compile());
 	
 	RemoveTypedElementRowHandleQuery = DataStorage.RegisterQuery(
 		Select()
 			.ReadOnly<FTEDSTypedElementColumn>()
 		.Compile());
-	
-	RemovePendingTypedElementRowHandleQuery = DataStorage.RegisterQuery(
-			Select().Where().All<FPendingPopulateTypedElementColumn>()
-		.Compile());
 }
 
 void UTEDSTypedElementBridge::UnregisterQuery_NewUObject(ITypedElementDataStorageInterface& DataStorage)
 {
-	DataStorage.UnregisterQuery(QueryHandle);
-	QueryHandle = TypedElementInvalidQueryHandle;
 }
 
 FOnTEDSTypedElementBridgeEnable& UTEDSTypedElementBridge::OnEnabled()
@@ -118,21 +93,6 @@ void UTEDSTypedElementBridge::CleanupTypedElementColumns(ITypedElementDataStorag
 		
 		DataStorage.BatchAddRemoveColumns(TConstArrayView<TypedElementRowHandle>(Handles), {}, {FTEDSTypedElementColumn::StaticStruct()});
 	}
-
-	// Remove any columns on row indicating they are pending population of their TEv1 columnn
-	{
-		TArray<TypedElementRowHandle> Handles;
-		using namespace TypedElementQueryBuilder;
-		DataStorage.RunQuery(
-			RemoveTypedElementRowHandleQuery,
-			CreateDirectQueryCallbackBinding(
-				[&Handles](ITypedElementDataStorageInterface::IDirectQueryContext& Context)
-				{
-					Handles.Append(Context.GetRowHandles());
-				}));
-		
-		DataStorage.BatchAddRemoveColumns(TConstArrayView<TypedElementRowHandle>(Handles), {}, {FPendingPopulateTypedElementColumn::StaticStruct()});
-	}
 }
 
 void UTEDSTypedElementBridge::HandleOnEnabled(IConsoleVariable* CVar)
@@ -142,9 +102,6 @@ void UTEDSTypedElementBridge::HandleOnEnabled(IConsoleVariable* CVar)
 
 	if (bIsEnabled)
 	{
-		DataStorage->UnregisterQuery(QueryHandle);
-		QueryHandle = TypedElementInvalidQueryHandle;
-		
 		RegisterQuery_NewUObject(*DataStorage);
 		UTEDSTypedElementBridge::OnEnabled().Broadcast(bIsEnabled);
 	}
diff --git a/Engine/Plugins/Experimental/TypedElementsDataStorage/Source/TypedElementsDataStorage/Public/Compatibility/TEDSTypedElementBridge.h b/Engine/Plugins/Experimental/TypedElementsDataStorage/Source/TypedElementsDataStorage/Public/Compatibility/TEDSTypedElementBridge.h
--- a/Engine/Plugins/Experimental/TypedElementsDataStorage/Source/TypedElementsDataStorage/Public/Compatibility/TEDSTypedElementBridge.h
+++ b/Engine/Plugins/Experimental/TypedElementsDataStorage/Source/TypedElementsDataStorage/Public/Compatibility/TEDSTypedElementBridge.h
@@ -37,9 +37,7 @@ private:
 	void CleanupTypedElementColumns(ITypedElementDataStorageInterface& DataStorage);
 	void HandleOnEnabled(IConsoleVariable* CVar);
 	
-	TypedElementQueryHandle QueryHandle = TypedElementDataStorage::InvalidQueryHandle;
 	TypedElementQueryHandle RemoveTypedElementRowHandleQuery = TypedElementDataStorage::InvalidQueryHandle;
-	TypedElementQueryHandle RemovePendingTypedElementRowHandleQuery = TypedElementDataStorage::InvalidQueryHandle;
 	FDelegateHandle DebugEnabledDelegateHandle;
 };
 
@@ -52,12 +50,4 @@ struct FTEDSTypedElementColumn : public FTypedElementDataStorageColumn
 	GENERATED_BODY()
 	
 	FTypedElementHandle Handle;
-};
-
-// Used by processors to defer assignment of TEDSTypedElementColumn
-// until processor runs
-USTRUCT()
-struct FPendingPopulateTypedElementColumn : public FTypedElementDataStorageTag
-{
-	GENERATED_BODY()
 };
\ No newline at end of file
