Subqueries now support adding, removing and checking columns.

The following additional changes were made:
- The "sync widgets" group is now automatically forced to the main thread and now also guaranteed to be behind the "default" group.
- The environment is now available inside the query and sub-query context.
- The environment now tracks the current update cycle (a.k.a. tick). This allows query callbacks to avoid running the same logic multiple times in the same frame.
- HasColumn now has an overloaded version that takes a row. This allows finding out if there's column on any row, not just the ones in the current context.
- FSingleRowQueryContextForwardToDataStorage has been removed in favor of using the same subquery context everywhere.

diff --git a/Engine/Plugins/Experimental/TypedElementsDataStorage/Source/TypedElementsDataStorage/Private/Processors/TypedElementProcessorAdaptors.cpp b/Engine/Plugins/Experimental/TypedElementsDataStorage/Source/TypedElementsDataStorage/Private/Processors/TypedElementProcessorAdaptors.cpp
--- a/Engine/Plugins/Experimental/TypedElementsDataStorage/Source/TypedElementsDataStorage/Private/Processors/TypedElementProcessorAdaptors.cpp
+++ b/Engine/Plugins/Experimental/TypedElementsDataStorage/Source/TypedElementsDataStorage/Private/Processors/TypedElementProcessorAdaptors.cpp
@@ -78,121 +78,52 @@ struct FMassContextCommon : public T
 			return Context.DoesArchetypeHaveFragment(*ColumnType);
 		}
 		const bool bIsTagOrFragment = false;
-		checkf(bIsTagOrFragment, TEXT("Attempting to check for a column type that is not a column or tag."))
+		checkf(bIsTagOrFragment, TEXT("Attempting to check for a column type that is not a column or tag."));
 		return false;
 	}
 
-protected:
-	explicit FMassContextCommon(FMassExecutionContext& InContext)
-		: Context(InContext)
-	{}
-
-	FMassExecutionContext& Context;
-};
-
-struct FMassDirectContextForwarder final : public FMassContextCommon<ITypedElementDataStorageInterface::IDirectQueryContext>
-{
-	explicit FMassDirectContextForwarder(FMassExecutionContext& InContext)
-		: FMassContextCommon(InContext)
-	{}
-
-	~FMassDirectContextForwarder() override = default;
-};
-
-struct FMassSubqueryContextForwarder final : public FMassContextCommon<ITypedElementDataStorageInterface::ISubqueryContext>
-{
-	explicit FMassSubqueryContextForwarder(FMassExecutionContext& InContext)
-		: FMassContextCommon(InContext)
-	{}
-
-	~FMassSubqueryContextForwarder() override = default;
-};
-
-struct FMassSingleRowSubqueryContextForwarder final : public ITypedElementDataStorageInterface::ISubqueryContext
-{
-	FMassSingleRowSubqueryContextForwarder(FMassEntityManager& InEntityManager, TypedElementDataStorage::RowHandle InRowHandle)
-		: EntityManager(InEntityManager)
-		, RowHandle(InRowHandle)
-	{}
-
-	~FMassSingleRowSubqueryContextForwarder() override = default;
-
-	const void* GetColumn(const UScriptStruct* ColumnType) const override
+	bool HasColumn(TypedElementDataStorage::RowHandle Row, const UScriptStruct* ColumnType) const override
 	{
-		return EntityManager.GetFragmentDataStruct(FMassEntityHandle::FromNumber(RowHandle), ColumnType).GetMemory();
-	}
+		FMassEntityHandle Entity = FMassEntityHandle::FromNumber(Row);
+		FMassEntityManager& Manager = Context.GetEntityManagerChecked();
+		FMassArchetypeHandle Archetype = Manager.GetArchetypeForEntity(Entity);
+		const FMassArchetypeCompositionDescriptor& Composition = Manager.GetArchetypeComposition(Archetype);
 
-	void* GetMutableColumn(const UScriptStruct* ColumnType) override
-	{
-		return EntityManager.GetFragmentDataStruct(FMassEntityHandle::FromNumber(RowHandle), ColumnType).GetMemory();
-	}
-
-	void GetColumns(TArrayView<char*> RetrievedAddresses, TConstArrayView<TWeakObjectPtr<const UScriptStruct>> ColumnTypes,
-		TConstArrayView<TypedElementDataStorage::EQueryAccessType> AccessTypes) override
-	{
-		checkf(RetrievedAddresses.Num() == ColumnTypes.Num(), TEXT("Unable to retrieve a batch of columns as the number of addresses "
-			"doesn't match the number of requested column."));
-		checkf(RetrievedAddresses.Num() == AccessTypes.Num(), TEXT("Unable to retrieve a batch of columns as the number of addresses "
-			"doesn't match the number of access types."));
-
-		GetColumnsUnguarded(ColumnTypes.Num(), RetrievedAddresses.GetData(), ColumnTypes.GetData(), AccessTypes.GetData());
-	}
-
-	void GetColumnsUnguarded(int32 TypeCount, char** RetrievedAddresses, const TWeakObjectPtr<const UScriptStruct>* ColumnTypes,
-		const TypedElementDataStorage::EQueryAccessType*) override
-	{
-		for (int32 Index = 0; Index < TypeCount; ++Index)
-		{
-			checkf(ColumnTypes->IsValid(), TEXT("Attempting to retrieve a column that is not available."));
-			*RetrievedAddresses = reinterpret_cast<char*>(GetMutableColumn(ColumnTypes->Get()));
-
-			++RetrievedAddresses;
-			++ColumnTypes;
-		}
-	}
-
-	uint32 GetRowCount() const override { return 1; }
-	TConstArrayView<TypedElementDataStorage::RowHandle> GetRowHandles() const override
-	{
-		return TConstArrayView<TypedElementDataStorage::RowHandle>(&RowHandle, 1);
-	}
-
-	bool HasColumn(const UScriptStruct* ColumnType) const override
-	{
-		FMassArchetypeHandle Archetype = EntityManager.GetArchetypeForEntity(FMassEntityHandle::FromNumber(RowHandle));
-		FMassArchetypeCompositionDescriptor Composition = EntityManager.GetArchetypeComposition(Archetype);
-		// Tags
 		if (ColumnType->IsChildOf(FMassTag::StaticStruct()))
 		{
 			return Composition.Tags.Contains(*ColumnType);
 		}
-		// Columns
 		if (ColumnType->IsChildOf(FMassFragment::StaticStruct()))
 		{
 			return Composition.Fragments.Contains(*ColumnType);
 		}
 		const bool bIsTagOrFragment = false;
-		checkf(bIsTagOrFragment, TEXT("Attempting to check for a column type that is not a column or tag."))
+		checkf(bIsTagOrFragment, TEXT("Attempting to check for a column type that is not a column or tag."));
 		return false;
 	}
 
-private:
-	FMassEntityManager& EntityManager;
-	TypedElementDataStorage::RowHandle RowHandle;
+protected:
+	explicit FMassContextCommon(FMassExecutionContext& InContext)
+		: Context(InContext)
+	{}
+
+	FMassExecutionContext& Context;
 };
 
-struct FMassContextForwarder final : public FMassContextCommon<ITypedElementDataStorageInterface::IQueryContext>
+template<typename T>
+struct FMassWithEnvironmentContextCommon : public FMassContextCommon<T>
 {
-private:
+	using Parent = FMassContextCommon<T>;
+protected:
 	void TedsColumnsToMassDescriptorIfActiveTable(
-		FMassArchetypeCompositionDescriptor& Descriptor, 
+		FMassArchetypeCompositionDescriptor& Descriptor,
 		TConstArrayView<const UScriptStruct*> ColumnTypes)
 	{
 		for (const UScriptStruct* ColumnType : ColumnTypes)
 		{
 			if (ColumnType->IsChildOf(FMassTag::StaticStruct()))
 			{
-				if (Context.DoesArchetypeHaveTag(*ColumnType))
+				if (this->Context.DoesArchetypeHaveTag(*ColumnType))
 				{
 					Descriptor.Tags.Add(*ColumnType);
 				}
@@ -201,7 +132,7 @@ private:
 			{
 				checkf(ColumnType->IsChildOf(FMassFragment::StaticStruct()),
 					TEXT("Given struct type is not a valid fragment or tag type."));
-				if (Context.DoesArchetypeHaveFragment(*ColumnType))
+				if (this->Context.DoesArchetypeHaveFragment(*ColumnType))
 				{
 					Descriptor.Fragments.Add(*ColumnType);
 				}
@@ -209,7 +140,9 @@ private:
 		}
 	}
 
-	void TedsColumnsToMassDescriptor(FMassArchetypeCompositionDescriptor& Descriptor, TConstArrayView<const UScriptStruct*> ColumnTypes)
+	void TedsColumnsToMassDescriptor(
+		FMassArchetypeCompositionDescriptor& Descriptor,
+		TConstArrayView<const UScriptStruct*> ColumnTypes)
 	{
 		for (const UScriptStruct* ColumnType : ColumnTypes)
 		{
@@ -226,80 +159,18 @@ private:
 			}
 		}
 	}
-public:
 
-	FMassContextForwarder(ITypedElementDataStorageInterface::FQueryDescription& InQueryDescription, FMassExecutionContext& InContext, 
-		FTypedElementExtendedQueryStore& InQueryStore, FTypedElementDatabaseEnvironment& InEnvironment)
-		: FMassContextCommon(InContext)
-		, QueryDescription(InQueryDescription)
-		, QueryStore(InQueryStore)
+public:
+	FMassWithEnvironmentContextCommon(FMassExecutionContext& InContext, FTypedElementDatabaseEnvironment& InEnvironment)
+		: FMassContextCommon<T>(InContext)
 		, Environment(InEnvironment)
 	{}
 
-	~FMassContextForwarder() override = default;
-
-	UObject* GetMutableDependency(const UClass* DependencyClass) override
-	{
-		return Context.GetMutableSubsystem<USubsystem>(const_cast<UClass*>(DependencyClass));
-	}
-
-	const UObject* GetDependency(const UClass* DependencyClass) override
-	{
-		return Context.GetSubsystem<USubsystem>(const_cast<UClass*>(DependencyClass));
-	}
-
-	void GetDependencies(TArrayView<UObject*> RetrievedAddresses, TConstArrayView<TWeakObjectPtr<const UClass>> SubsystemTypes,
-		TConstArrayView<ITypedElementDataStorageInterface::EQueryAccessType> AccessTypes) override
-	{
-		checkf(RetrievedAddresses.Num() == SubsystemTypes.Num(), TEXT("Unable to retrieve a batch of subsystem as the number of addresses "
-			"doesn't match the number of requested subsystem types."));
-
-		GetDependenciesUnguarded(RetrievedAddresses.Num(), RetrievedAddresses.GetData(), SubsystemTypes.GetData(), AccessTypes.GetData());
-	}
-
-	void GetDependenciesUnguarded(int32 SubsystemCount, UObject** RetrievedAddresses, const TWeakObjectPtr<const UClass>* DependencyTypes,
-		const ITypedElementDataStorageInterface::EQueryAccessType* AccessTypes)
-	{
-		for (int32 Index = 0; Index < SubsystemCount; ++Index)
-		{
-			checkf(DependencyTypes->IsValid(), TEXT("Attempting to retrieve a subsystem that's no longer valid."));
-			*RetrievedAddresses = *AccessTypes == ITypedElementDataStorageInterface::EQueryAccessType::ReadOnly
-				? const_cast<USubsystem*>(Context.GetSubsystem<USubsystem>(const_cast<UClass*>(DependencyTypes->Get())))
-				: Context.GetMutableSubsystem<USubsystem>(const_cast<UClass*>(DependencyTypes->Get()));
-
-			++RetrievedAddresses;
-			++DependencyTypes;
-			++AccessTypes;
-		}
-	}
-
-	bool IsRowAvailable(TypedElementDataStorage::RowHandle Row) const override
-	{
-		return Environment.GetMassEntityManager().IsEntityValid(FMassEntityHandle::FromNumber(Row));
-	}
-
-	bool HasRowBeenAssigned(TypedElementDataStorage::RowHandle Row) const override
-	{
-		return Environment.GetMassEntityManager().IsEntityActive(FMassEntityHandle::FromNumber(Row));
-	}
-
-	void RemoveRow(TypedElementRowHandle Row) override
-	{
-		Context.Defer().DestroyEntity(FMassEntityHandle::FromNumber(Row));
-	}
+	~FMassWithEnvironmentContextCommon() override = default;
 
-	void RemoveRows(TConstArrayView<TypedElementRowHandle> Rows) override
+	uint64 GetUpdateCycleId() const override
 	{
-		// Row handles and entities map 1:1 for data, so a reintpret_cast can be safely done to avoid
-		// having to allocate memory and iterating over the rows.
-
-		static_assert(sizeof(FMassEntityHandle) == sizeof(TypedElementRowHandle), 
-			"Sizes of mass entity and data storage row have gone out of sync.");
-		static_assert(alignof(FMassEntityHandle) == alignof(TypedElementRowHandle),
-			"Alignment of mass entity and data storage row have gone out of sync.");
-
-		Context.Defer().DestroyEntities(
-			TConstArrayView<FMassEntityHandle>(reinterpret_cast<const FMassEntityHandle*>(Rows.begin()), Rows.Num()));
+		return Environment.GetUpdateCycleId();
 	}
 
 	void* AddColumnUninitialized(TypedElementDataStorage::RowHandle Row, const UScriptStruct* ObjectType) override
@@ -311,19 +182,19 @@ public:
 			});
 	}
 
-	void* AddColumnUninitialized(TypedElementDataStorage::RowHandle Row, const UScriptStruct* ObjectType, ObjectCopyOrMove Relocator) override
+	void* AddColumnUninitialized(TypedElementDataStorage::RowHandle Row, const UScriptStruct* ObjectType, typename Parent::ObjectCopyOrMove Relocator) override
 	{
 		checkf(ObjectType->IsChildOf(FMassFragment::StaticStruct()), TEXT("Column [%s] can not be a tag"), *ObjectType->GetName());
 		
 		struct FAddMoveableValueColumn
 		{
-			ObjectCopyOrMove Relocator;
+			typename Parent::ObjectCopyOrMove Relocator;
 			const UScriptStruct* FragmentType;
 			FMassEntityHandle Entity;
 			void* Object;
 
 			FAddMoveableValueColumn() = default;
-			FAddMoveableValueColumn(ObjectCopyOrMove InRelocator, const UScriptStruct* InFragmentType, FMassEntityHandle InEntity, void* InObject)
+			FAddMoveableValueColumn(typename Parent::ObjectCopyOrMove InRelocator, const UScriptStruct* InFragmentType, FMassEntityHandle InEntity, void* InObject)
 				: Relocator(InRelocator)
 				, FragmentType(InFragmentType)
 				, Entity(InEntity)
@@ -334,13 +205,13 @@ public:
 		struct FAddMoveableValueColumnWithDestructor : FAddMoveableValueColumn
 		{
 			FAddMoveableValueColumnWithDestructor(
-				ObjectCopyOrMove InRelocator, const UScriptStruct* InFragmentType, FMassEntityHandle InEntity, void* InObject)
+				Parent::ObjectCopyOrMove InRelocator, const UScriptStruct* InFragmentType, FMassEntityHandle InEntity, void* InObject)
 				: FAddMoveableValueColumn(InRelocator, InFragmentType, InEntity, InObject)
 			{}
 
 			~FAddMoveableValueColumnWithDestructor()
 			{
-				FragmentType->DestroyStruct(Object);
+				this->FragmentType->DestroyStruct(this->Object);
 			}
 		};
 
@@ -356,7 +227,7 @@ public:
 			AddedColumn = ScratchBuffer.Emplace<FAddMoveableValueColumnWithDestructor>(Relocator, ObjectType, FMassEntityHandle::FromNumber(Row), ColumnData);
 		}
 
-		Context.Defer().PushCommand<FMassDeferredAddCommand>(
+		this->Context.Defer().template PushCommand<FMassDeferredAddCommand>(
 			[AddedColumn](FMassEntityManager& System)
 			{
 				// Check entity before proceeding. It's possible it may have been invalidated before this deferred call fired.
@@ -372,7 +243,7 @@ public:
 				}
 			});
 
-		Context.Defer().PushCommand<FMassDeferredSetCommand>(
+		this->Context.Defer().template PushCommand<FMassDeferredSetCommand>(
 			[AddedColumn](FMassEntityManager& System)
 			{
 				// Check entity before proceeding. It's possible it may have been invalidated before this deferred call fired.
@@ -398,10 +269,13 @@ public:
 		TedsColumnsToMassDescriptor(AddedColumns->AddDescriptor, ColumnTypes);
 		AddedColumns->Entity = FMassEntityHandle::FromNumber(Row);
 
-		Context.Defer().PushCommand<FMassDeferredAddCommand>(
+		this->Context.Defer().template PushCommand<FMassDeferredAddCommand>(
 			[AddedColumns](FMassEntityManager& System)
 			{
-				System.AddCompositionToEntity_GetDelta(AddedColumns->Entity, AddedColumns->AddDescriptor);
+				if (System.IsEntityValid(AddedColumns->Entity))
+				{
+					System.AddCompositionToEntity_GetDelta(AddedColumns->Entity, AddedColumns->AddDescriptor);
+				}
 			});
 	}
 
@@ -427,14 +301,17 @@ public:
 		}
 		AddedColumns->EntityCount = Rows.Num();
 
-		Context.Defer().PushCommand<FMassDeferredAddCommand>(
+		this->Context.Defer().template PushCommand<FMassDeferredAddCommand>(
 			[AddedColumns](FMassEntityManager& System)
 			{
 				FMassEntityHandle* Entities = AddedColumns->Entities;
 				int32 Count = AddedColumns->EntityCount;
 				for (int32 Counter = 0; Counter < Count; ++Counter)
 				{
-					System.AddCompositionToEntity_GetDelta(*Entities++, AddedColumns->AddDescriptor);
+					if (System.IsEntityValid(*Entities))
+					{
+						System.AddCompositionToEntity_GetDelta(*Entities++, AddedColumns->AddDescriptor);
+					}
 				}
 			});
 	}
@@ -449,13 +326,19 @@ public:
 
 		FRemovedColumns* RemovedColumns = Environment.GetScratchBuffer().Emplace<FRemovedColumns>();
 		TedsColumnsToMassDescriptorIfActiveTable(RemovedColumns->RemoveDescriptor, ColumnTypes);
-		RemovedColumns->Entity = FMassEntityHandle::FromNumber(Row);
+		if (!RemovedColumns->RemoveDescriptor.IsEmpty())
+		{
+			RemovedColumns->Entity = FMassEntityHandle::FromNumber(Row);
 
-		Context.Defer().PushCommand<FMassDeferredAddCommand>(
-			[RemovedColumns](FMassEntityManager& System)
-			{
-				System.RemoveCompositionFromEntity(RemovedColumns->Entity, RemovedColumns->RemoveDescriptor);
-			});
+			this->Context.Defer().template PushCommand<FMassDeferredAddCommand>(
+				[RemovedColumns](FMassEntityManager& System)
+				{
+					if (System.IsEntityValid(RemovedColumns->Entity))
+					{
+						System.RemoveCompositionFromEntity(RemovedColumns->Entity, RemovedColumns->RemoveDescriptor);
+					}
+				});
+		}
 	}
 
 	void RemoveColumns(TConstArrayView<TypedElementDataStorage::RowHandle> Rows, TConstArrayView<const UScriptStruct*> ColumnTypes) override
@@ -480,18 +363,119 @@ public:
 		}
 		RemovedColumns->EntityCount = Rows.Num();
 
-		Context.Defer().PushCommand<FMassDeferredAddCommand>(
+		this->Context.Defer().template PushCommand<FMassDeferredAddCommand>(
 			[RemovedColumns](FMassEntityManager& System)
 			{
 				FMassEntityHandle* Entities = RemovedColumns->Entities;
 				int32 Count = RemovedColumns->EntityCount;
 				for (int32 Counter = 0; Counter < Count; ++Counter)
 				{
-					System.RemoveCompositionFromEntity(*Entities++, RemovedColumns->RemoveDescriptor);
+					if (System.IsEntityValid(*Entities))
+					{
+						System.RemoveCompositionFromEntity(*Entities++, RemovedColumns->RemoveDescriptor);
+					}
 				}
 			});
 	}
 
+protected:
+	FTypedElementDatabaseEnvironment& Environment;
+};
+
+struct FMassDirectContextForwarder final : public FMassContextCommon<ITypedElementDataStorageInterface::IDirectQueryContext>
+{
+	explicit FMassDirectContextForwarder(FMassExecutionContext& InContext)
+		: FMassContextCommon(InContext)
+	{}
+
+	~FMassDirectContextForwarder() override = default;
+};
+
+struct FMassSubqueryContextForwarder final : public FMassWithEnvironmentContextCommon<ITypedElementDataStorageInterface::ISubqueryContext>
+{
+	FMassSubqueryContextForwarder(FMassExecutionContext& InContext, FTypedElementDatabaseEnvironment& InEnvironment)
+		: FMassWithEnvironmentContextCommon(InContext, InEnvironment)
+	{}
+
+	~FMassSubqueryContextForwarder() override = default;
+};
+
+struct FMassContextForwarder final : public FMassWithEnvironmentContextCommon<ITypedElementDataStorageInterface::IQueryContext>
+{
+public:
+	FMassContextForwarder(ITypedElementDataStorageInterface::FQueryDescription& InQueryDescription, FMassExecutionContext& InContext, 
+		FTypedElementExtendedQueryStore& InQueryStore, FTypedElementDatabaseEnvironment& InEnvironment)
+		: FMassWithEnvironmentContextCommon(InContext, InEnvironment)
+		, QueryDescription(InQueryDescription)
+		, QueryStore(InQueryStore)
+	{}
+
+	~FMassContextForwarder() override = default;
+
+	UObject* GetMutableDependency(const UClass* DependencyClass) override
+	{
+		return Context.GetMutableSubsystem<USubsystem>(const_cast<UClass*>(DependencyClass));
+	}
+
+	const UObject* GetDependency(const UClass* DependencyClass) override
+	{
+		return Context.GetSubsystem<USubsystem>(const_cast<UClass*>(DependencyClass));
+	}
+
+	void GetDependencies(TArrayView<UObject*> RetrievedAddresses, TConstArrayView<TWeakObjectPtr<const UClass>> SubsystemTypes,
+		TConstArrayView<ITypedElementDataStorageInterface::EQueryAccessType> AccessTypes) override
+	{
+		checkf(RetrievedAddresses.Num() == SubsystemTypes.Num(), TEXT("Unable to retrieve a batch of subsystem as the number of addresses "
+			"doesn't match the number of requested subsystem types."));
+
+		GetDependenciesUnguarded(RetrievedAddresses.Num(), RetrievedAddresses.GetData(), SubsystemTypes.GetData(), AccessTypes.GetData());
+	}
+
+	void GetDependenciesUnguarded(int32 SubsystemCount, UObject** RetrievedAddresses, const TWeakObjectPtr<const UClass>* DependencyTypes,
+		const ITypedElementDataStorageInterface::EQueryAccessType* AccessTypes)
+	{
+		for (int32 Index = 0; Index < SubsystemCount; ++Index)
+		{
+			checkf(DependencyTypes->IsValid(), TEXT("Attempting to retrieve a subsystem that's no longer valid."));
+			*RetrievedAddresses = *AccessTypes == ITypedElementDataStorageInterface::EQueryAccessType::ReadOnly
+				? const_cast<USubsystem*>(Context.GetSubsystem<USubsystem>(const_cast<UClass*>(DependencyTypes->Get())))
+				: Context.GetMutableSubsystem<USubsystem>(const_cast<UClass*>(DependencyTypes->Get()));
+
+			++RetrievedAddresses;
+			++DependencyTypes;
+			++AccessTypes;
+		}
+	}
+
+	bool IsRowAvailable(TypedElementDataStorage::RowHandle Row) const override
+	{
+		return Environment.GetMassEntityManager().IsEntityValid(FMassEntityHandle::FromNumber(Row));
+	}
+
+	bool HasRowBeenAssigned(TypedElementDataStorage::RowHandle Row) const override
+	{
+		return Environment.GetMassEntityManager().IsEntityActive(FMassEntityHandle::FromNumber(Row));
+	}
+
+	void RemoveRow(TypedElementRowHandle Row) override
+	{
+		Context.Defer().DestroyEntity(FMassEntityHandle::FromNumber(Row));
+	}
+
+	void RemoveRows(TConstArrayView<TypedElementRowHandle> Rows) override
+	{
+		// Row handles and entities map 1:1 for data, so a reintpret_cast can be safely done to avoid
+		// having to allocate memory and iterating over the rows.
+
+		static_assert(sizeof(FMassEntityHandle) == sizeof(TypedElementRowHandle), 
+			"Sizes of mass entity and data storage row have gone out of sync.");
+		static_assert(alignof(FMassEntityHandle) == alignof(TypedElementRowHandle),
+			"Alignment of mass entity and data storage row have gone out of sync.");
+
+		Context.Defer().DestroyEntities(
+			TConstArrayView<FMassEntityHandle>(reinterpret_cast<const FMassEntityHandle*>(Rows.begin()), Rows.Num()));
+	}
+
 	TypedElementDataStorage::RowHandle FindIndexedRow(TypedElementDataStorage::IndexHash Index) const override
 	{
 		return Environment.GetIndexTable().FindIndexedRow(Index);
@@ -517,7 +501,7 @@ public:
 		{
 			const TypedElementQueryHandle SubqueryHandle = QueryDescription.Subqueries[SubqueryIndex];
 			const FTypedElementExtendedQueryStore::Handle StorageHandle(SubqueryHandle);
-			return QueryStore.RunQuery(Context.GetEntityManagerChecked(), StorageHandle, Callback);
+			return QueryStore.RunQuery(Context.GetEntityManagerChecked(), Environment, Context, StorageHandle, Callback);
 		}
 		else
 		{
@@ -532,7 +516,7 @@ public:
 		{
 			const TypedElementQueryHandle SubqueryHandle = QueryDescription.Subqueries[SubqueryIndex];
 			const FTypedElementExtendedQueryStore::Handle StorageHandle(SubqueryHandle);
-			return QueryStore.RunQuery(Context.GetEntityManagerChecked(), StorageHandle, Row, Callback);
+			return QueryStore.RunQuery(Context.GetEntityManagerChecked(), Environment, Context, StorageHandle, Row, Callback);
 		}
 		else
 		{
@@ -542,7 +526,6 @@ public:
 
 	ITypedElementDataStorageInterface::FQueryDescription& QueryDescription;
 	FTypedElementExtendedQueryStore& QueryStore;
-	FTypedElementDatabaseEnvironment& Environment;
 };
 
 
@@ -697,7 +680,8 @@ TypedElementDataStorage::FQueryResult FTypedElementQueryProcessorData::Execute(
 	TypedElementDataStorage::DirectQueryCallbackRef& Callback,
 	TypedElementDataStorage::FQueryDescription& Description,
 	FMassEntityQuery& NativeQuery, 
-	FMassEntityManager& EntityManager)
+	FMassEntityManager& EntityManager,
+	FTypedElementDatabaseEnvironment& Environment)
 {
 	FMassExecutionContext Context(EntityManager);
 	ITypedElementDataStorageInterface::FQueryResult Result;
@@ -719,17 +703,20 @@ TypedElementDataStorage::FQueryResult FTypedElementQueryProcessorData::Execute(
 	TypedElementDataStorage::SubqueryCallbackRef& Callback,
 	TypedElementDataStorage::FQueryDescription& Description,
 	FMassEntityQuery& NativeQuery,
-	FMassEntityManager& EntityManager)
+	FMassEntityManager& EntityManager,
+	FTypedElementDatabaseEnvironment& Environment,
+	FMassExecutionContext& ParentContext)
 {
 	FMassExecutionContext Context(EntityManager);
+	Context.SetDeferredCommandBuffer(ParentContext.GetSharedDeferredCommandBuffer());
 	ITypedElementDataStorageInterface::FQueryResult Result;
 	Result.Completed = ITypedElementDataStorageInterface::FQueryResult::ECompletion::Fully;
 
 	NativeQuery.ForEachEntityChunk(EntityManager, Context,
-		[&Result, &Callback, &Description](FMassExecutionContext& Context)
+		[&Result, &Callback, &Description, &Environment](FMassExecutionContext& Context)
 		{
 			// No need to cache any subsystem dependencies as these are not accessible from a subquery.
-			FMassSubqueryContextForwarder QueryContext(Context);
+			FMassSubqueryContextForwarder QueryContext(Context, Environment);
 			Callback(Description, QueryContext);
 			Result.Count += Context.GetNumEntities();
 		}
@@ -742,7 +729,9 @@ TypedElementDataStorage::FQueryResult FTypedElementQueryProcessorData::Execute(
 	TypedElementDataStorage::FQueryDescription& Description,
 	TypedElementDataStorage::RowHandle RowHandle,
 	FMassEntityQuery& NativeQuery,
-	FMassEntityManager& EntityManager)
+	FMassEntityManager& EntityManager,
+	FTypedElementDatabaseEnvironment& Environment,
+	FMassExecutionContext& ParentContext)
 {
 	ITypedElementDataStorageInterface::FQueryResult Result;
 	Result.Completed = ITypedElementDataStorageInterface::FQueryResult::ECompletion::Fully;
@@ -751,12 +740,20 @@ TypedElementDataStorage::FQueryResult FTypedElementQueryProcessorData::Execute(
 	if (EntityManager.IsEntityActive(NativeEntity))
 	{
 		FMassArchetypeHandle NativeArchetype = EntityManager.GetArchetypeForEntityUnsafe(NativeEntity);
-		if (NativeQuery.DoesArchetypeMatchRequirements(NativeArchetype))
-		{
-			FMassSingleRowSubqueryContextForwarder QueryContext(EntityManager, RowHandle);
-			Callback(Description, QueryContext);
-			Result.Count = 1;
-		}
+		FMassExecutionContext Context(EntityManager);
+		Context.SetEntityCollection(FMassArchetypeEntityCollection(NativeArchetype, { NativeEntity }, FMassArchetypeEntityCollection::NoDuplicates));
+		Context.SetDeferredCommandBuffer(ParentContext.GetSharedDeferredCommandBuffer());
+
+		NativeQuery.ForEachEntityChunk(EntityManager, Context,
+			[&Result, &Callback, &Description, &Environment, &EntityManager, RowHandle](FMassExecutionContext& Context)
+			{
+				// No need to cache any subsystem dependencies as these are not accessible from a subquery.
+				FMassSubqueryContextForwarder QueryContext(Context, Environment);
+				Callback(Description, QueryContext);
+				Result.Count += Context.GetNumEntities();
+			}
+		);
+		checkf(Result.Count < 2, TEXT("Single row subquery produced multiple results."));
 	}
 	return Result;
 }
diff --git a/Engine/Plugins/Experimental/TypedElementsDataStorage/Source/TypedElementsDataStorage/Private/Processors/TypedElementProcessorAdaptors.h b/Engine/Plugins/Experimental/TypedElementsDataStorage/Source/TypedElementsDataStorage/Private/Processors/TypedElementProcessorAdaptors.h
--- a/Engine/Plugins/Experimental/TypedElementsDataStorage/Source/TypedElementsDataStorage/Private/Processors/TypedElementProcessorAdaptors.h
+++ b/Engine/Plugins/Experimental/TypedElementsDataStorage/Source/TypedElementsDataStorage/Private/Processors/TypedElementProcessorAdaptors.h
@@ -53,18 +53,23 @@ struct FTypedElementQueryProcessorData
 		TypedElementDataStorage::DirectQueryCallbackRef& Callback,
 		TypedElementDataStorage::FQueryDescription& Description,
 		FMassEntityQuery& NativeQuery, 
-		FMassEntityManager& EntityManager);
+		FMassEntityManager& EntityManager,
+		FTypedElementDatabaseEnvironment& Environment);
 	static TypedElementDataStorage::FQueryResult Execute(
 		TypedElementDataStorage::SubqueryCallbackRef& Callback,
 		TypedElementDataStorage::FQueryDescription& Description,
 		FMassEntityQuery& NativeQuery,
-		FMassEntityManager& EntityManager);
+		FMassEntityManager& EntityManager,
+		FTypedElementDatabaseEnvironment& Environment,
+		FMassExecutionContext& ParentContext);
 	static TypedElementDataStorage::FQueryResult Execute(
 		TypedElementDataStorage::SubqueryCallbackRef& Callback,
 		TypedElementDataStorage::FQueryDescription& Description,
 		TypedElementDataStorage::RowHandle RowHandle,
 		FMassEntityQuery& NativeQuery,
-		FMassEntityManager& EntityManager);
+		FMassEntityManager& EntityManager,
+		FTypedElementDatabaseEnvironment& Environment,
+		FMassExecutionContext& ParentContext);
 	void Execute(FMassEntityManager& EntityManager, FMassExecutionContext& Context);
 
 	static bool PrepareCachedDependenciesOnQuery(
diff --git a/Engine/Plugins/Experimental/TypedElementsDataStorage/Source/TypedElementsDataStorage/Private/Queries/TypedElementExtendedQueryStore.cpp b/Engine/Plugins/Experimental/TypedElementsDataStorage/Source/TypedElementsDataStorage/Private/Queries/TypedElementExtendedQueryStore.cpp
--- a/Engine/Plugins/Experimental/TypedElementsDataStorage/Source/TypedElementsDataStorage/Private/Queries/TypedElementExtendedQueryStore.cpp
+++ b/Engine/Plugins/Experimental/TypedElementsDataStorage/Source/TypedElementsDataStorage/Private/Queries/TypedElementExtendedQueryStore.cpp
@@ -172,7 +172,11 @@ TypedElementDataStorage::FQueryResult FTypedElementExtendedQueryStore::RunQuery(
 }
 
 template<typename CallbackReference>
-TypedElementDataStorage::FQueryResult FTypedElementExtendedQueryStore::RunQueryCallbackCommon(FMassEntityManager& EntityManager, Handle Query,
+TypedElementDataStorage::FQueryResult FTypedElementExtendedQueryStore::RunQueryCallbackCommon(
+	FMassEntityManager& EntityManager, 
+	FTypedElementDatabaseEnvironment& Environment,
+	FMassExecutionContext* ParentContext,
+	Handle Query,
 	CallbackReference Callback)
 {
 	using ActionType = ITypedElementDataStorageInterface::FQueryDescription::EActionType;
@@ -190,8 +194,16 @@ TypedElementDataStorage::FQueryResult FTypedElementExtendedQueryStore::RunQueryC
 		case ActionType::Select:
 			if (!QueryData->Processor.IsValid())
 			{
-				Result = FTypedElementQueryProcessorData::Execute(
-					Callback, QueryData->Description, QueryData->NativeQuery, EntityManager);
+				if constexpr (std::is_same_v<CallbackReference, TypedElementDataStorage::DirectQueryCallbackRef>)
+				{
+					Result = FTypedElementQueryProcessorData::Execute(
+						Callback, QueryData->Description, QueryData->NativeQuery, EntityManager, Environment);
+				}
+				else
+				{
+					Result = FTypedElementQueryProcessorData::Execute(
+						Callback, QueryData->Description, QueryData->NativeQuery, EntityManager, Environment, *ParentContext);
+				}
 			}
 			else
 			{
@@ -216,20 +228,22 @@ TypedElementDataStorage::FQueryResult FTypedElementExtendedQueryStore::RunQueryC
 	return Result;
 }
 
-TypedElementDataStorage::FQueryResult FTypedElementExtendedQueryStore::RunQuery(FMassEntityManager& EntityManager, Handle Query,
-	TypedElementDataStorage::DirectQueryCallbackRef Callback)
+TypedElementDataStorage::FQueryResult FTypedElementExtendedQueryStore::RunQuery(FMassEntityManager& EntityManager, 
+	FTypedElementDatabaseEnvironment& Environment, Handle Query, TypedElementDataStorage::DirectQueryCallbackRef Callback)
 {
-	return RunQueryCallbackCommon(EntityManager, Query, Callback);
+	return RunQueryCallbackCommon(EntityManager, Environment, nullptr, Query, Callback);
 }
 
-TypedElementDataStorage::FQueryResult FTypedElementExtendedQueryStore::RunQuery(FMassEntityManager& EntityManager, Handle Query,
+TypedElementDataStorage::FQueryResult FTypedElementExtendedQueryStore::RunQuery(FMassEntityManager& EntityManager, 
+	FTypedElementDatabaseEnvironment& Environment, FMassExecutionContext& ParentContext, Handle Query, 
 	TypedElementDataStorage::SubqueryCallbackRef Callback)
 {
-	return RunQueryCallbackCommon(EntityManager, Query, Callback);
+	return RunQueryCallbackCommon(EntityManager, Environment, &ParentContext, Query, Callback);
 }
 
-TypedElementDataStorage::FQueryResult FTypedElementExtendedQueryStore::RunQuery(FMassEntityManager& EntityManager, Handle Query,
-	TypedElementRowHandle Row, TypedElementDataStorage::SubqueryCallbackRef Callback)
+TypedElementDataStorage::FQueryResult FTypedElementExtendedQueryStore::RunQuery(FMassEntityManager& EntityManager, 
+	FTypedElementDatabaseEnvironment& Environment, FMassExecutionContext& ParentContext, Handle Query, TypedElementRowHandle Row,
+	TypedElementDataStorage::SubqueryCallbackRef Callback)
 {
 	using ActionType = ITypedElementDataStorageInterface::FQueryDescription::EActionType;
 	using CompletionType = ITypedElementDataStorageInterface::FQueryResult::ECompletion;
@@ -247,7 +261,7 @@ TypedElementDataStorage::FQueryResult FTypedElementExtendedQueryStore::RunQuery(
 			if (!QueryData->Processor.IsValid())
 			{
 				Result = FTypedElementQueryProcessorData::Execute(
-					Callback, QueryData->Description, Row, QueryData->NativeQuery, EntityManager);
+					Callback, QueryData->Description, Row, QueryData->NativeQuery, EntityManager, Environment, ParentContext);
 			}
 			else
 			{
diff --git a/Engine/Plugins/Experimental/TypedElementsDataStorage/Source/TypedElementsDataStorage/Private/Queries/TypedElementExtendedQueryStore.h b/Engine/Plugins/Experimental/TypedElementsDataStorage/Source/TypedElementsDataStorage/Private/Queries/TypedElementExtendedQueryStore.h
--- a/Engine/Plugins/Experimental/TypedElementsDataStorage/Source/TypedElementsDataStorage/Private/Queries/TypedElementExtendedQueryStore.h
+++ b/Engine/Plugins/Experimental/TypedElementsDataStorage/Source/TypedElementsDataStorage/Private/Queries/TypedElementExtendedQueryStore.h
@@ -88,12 +88,24 @@ public:
 	 * @description Various functions to run queries.
 	 */
 	TypedElementDataStorage::FQueryResult RunQuery(FMassEntityManager& EntityManager, Handle Query);
-	TypedElementDataStorage::FQueryResult RunQuery(FMassEntityManager& EntityManager, Handle Query,
+	TypedElementDataStorage::FQueryResult RunQuery(
+		FMassEntityManager& EntityManager, 
+		FTypedElementDatabaseEnvironment& Environment,
+		Handle Query,
 		TypedElementDataStorage::DirectQueryCallbackRef Callback);
-	TypedElementDataStorage::FQueryResult RunQuery(FMassEntityManager& EntityManager, Handle Query,
+	TypedElementDataStorage::FQueryResult RunQuery(
+		FMassEntityManager& EntityManager,
+		FTypedElementDatabaseEnvironment& Environment,
+		FMassExecutionContext& ParentContext,
+		Handle Query,
+		TypedElementDataStorage::SubqueryCallbackRef Callback);
+	TypedElementDataStorage::FQueryResult RunQuery(
+		FMassEntityManager& EntityManager, 
+		FTypedElementDatabaseEnvironment& Environment,
+		FMassExecutionContext& ParentContext,
+		Handle Query,
+		TypedElementRowHandle Row, 
 		TypedElementDataStorage::SubqueryCallbackRef Callback);
-	TypedElementDataStorage::FQueryResult RunQuery(FMassEntityManager& EntityManager, Handle Query,
-		TypedElementRowHandle Row, TypedElementDataStorage::SubqueryCallbackRef Callback);
 	void RunPhasePreambleQueries(
 		FMassEntityManager& EntityManager, 
 		FTypedElementDatabaseEnvironment& Environment,
@@ -129,7 +141,12 @@ private:
 	};
 
 	template<typename CallbackReference>
-	TypedElementDataStorage::FQueryResult RunQueryCallbackCommon(FMassEntityManager& EntityManager, Handle Query, CallbackReference Callback);
+	TypedElementDataStorage::FQueryResult RunQueryCallbackCommon(
+		FMassEntityManager& EntityManager, 
+		FTypedElementDatabaseEnvironment& Environment,
+		FMassExecutionContext* ParentContext,
+		Handle Query, 
+		CallbackReference Callback);
 
 	FMassEntityQuery& SetupNativeQuery(ITypedElementDataStorageInterface::FQueryDescription& Query, FTypedElementExtendedQuery& StoredQuery);
 	bool SetupSelectedColumns(ITypedElementDataStorageInterface::FQueryDescription& Query, FMassEntityQuery& NativeQuery);
diff --git a/Engine/Plugins/Experimental/TypedElementsDataStorage/Source/TypedElementsDataStorage/Private/TypedElementDatabase.cpp b/Engine/Plugins/Experimental/TypedElementsDataStorage/Source/TypedElementsDataStorage/Private/TypedElementDatabase.cpp
--- a/Engine/Plugins/Experimental/TypedElementsDataStorage/Source/TypedElementsDataStorage/Private/TypedElementDatabase.cpp
+++ b/Engine/Plugins/Experimental/TypedElementsDataStorage/Source/TypedElementsDataStorage/Private/TypedElementDatabase.cpp
@@ -130,9 +130,12 @@ void UTypedElementDatabase::Initialize()
 			// Guarantee that syncing to the data storage always happens before syncing to external.
 			RegisterTickGroup(GetQueryTickGroupName(EQueryTickGroups::SyncExternalToDataStorage),
 				Phase, GetQueryTickGroupName(EQueryTickGroups::SyncDataStorageToExternal), {}, false);
-			// Guarantee that widgets syncs happen after external data has been updated to the data storage.
+			// Guarantee that widgets syncs happen after external data has been updated and internal processing 
+			// completed to the data storage.
 			RegisterTickGroup(GetQueryTickGroupName(EQueryTickGroups::SyncWidgets),
-				Phase, {}, GetQueryTickGroupName(EQueryTickGroups::SyncExternalToDataStorage), false);
+				Phase, {}, GetQueryTickGroupName(EQueryTickGroups::Default), true /* Needs main thread*/);
+			RegisterTickGroup(GetQueryTickGroupName(EQueryTickGroups::SyncWidgets),
+				Phase, {}, GetQueryTickGroupName(EQueryTickGroups::SyncExternalToDataStorage), true /* Needs main thread*/);
 		}
 	}
 }
@@ -823,7 +826,7 @@ ITypedElementDataStorageInterface::FQueryResult UTypedElementDatabase::RunQuery(
 	if (ActiveEditorEntityManager)
 	{
 		const FTypedElementExtendedQueryStore::Handle StorageHandle(Query);
-		return Queries.RunQuery(*ActiveEditorEntityManager, StorageHandle, Callback);
+		return Queries.RunQuery(*ActiveEditorEntityManager, *Environment, StorageHandle, Callback);
 	}
 	else
 	{
@@ -875,6 +878,7 @@ void UTypedElementDatabase::PreparePhase(EQueryTickPhase Phase, float DeltaTime)
 {
 	if (ActiveEditorEntityManager)
 	{
+		Environment->NextUpdateCycle();
 		Queries.RunPhasePreambleQueries(*ActiveEditorEntityManager, *Environment, Phase, DeltaTime);
 	}
 }
diff --git a/Engine/Plugins/Experimental/TypedElementsDataStorage/Source/TypedElementsDataStorage/Private/TypedElementDatabaseEnvironment.cpp b/Engine/Plugins/Experimental/TypedElementsDataStorage/Source/TypedElementsDataStorage/Private/TypedElementDatabaseEnvironment.cpp
--- a/Engine/Plugins/Experimental/TypedElementsDataStorage/Source/TypedElementsDataStorage/Private/TypedElementDatabaseEnvironment.cpp
+++ b/Engine/Plugins/Experimental/TypedElementsDataStorage/Source/TypedElementsDataStorage/Private/TypedElementDatabaseEnvironment.cpp
@@ -49,3 +49,13 @@ const FMassProcessingPhaseManager& FTypedElementDatabaseEnvironment::GetMassPhas
 {
 	return MassPhaseManager;
 }
+
+void FTypedElementDatabaseEnvironment::NextUpdateCycle()
+{
+	UpdateCycleId++;
+}
+
+uint64 FTypedElementDatabaseEnvironment::GetUpdateCycleId() const
+{
+	return UpdateCycleId;
+}
diff --git a/Engine/Plugins/Experimental/TypedElementsDataStorage/Source/TypedElementsDataStorage/Private/TypedElementDatabaseEnvironment.h b/Engine/Plugins/Experimental/TypedElementsDataStorage/Source/TypedElementsDataStorage/Private/TypedElementDatabaseEnvironment.h
--- a/Engine/Plugins/Experimental/TypedElementsDataStorage/Source/TypedElementsDataStorage/Private/TypedElementDatabaseEnvironment.h
+++ b/Engine/Plugins/Experimental/TypedElementsDataStorage/Source/TypedElementsDataStorage/Private/TypedElementDatabaseEnvironment.h
@@ -24,10 +24,15 @@ public:
 	FMassProcessingPhaseManager& GetMassPhaseManager();
 	const FMassProcessingPhaseManager& GetMassPhaseManager() const;
 
+	void NextUpdateCycle();
+	uint64 GetUpdateCycleId() const;
+
 private:
 	FTypedElementDatabaseIndexTable IndexTable;
 	TSharedPtr<FTypedElementDatabaseScratchBuffer> ScratchBuffer;
 	
 	FMassEntityManager& MassEntityManager;
 	FMassProcessingPhaseManager& MassPhaseManager;
+
+	uint64 UpdateCycleId = 0;
 };
diff --git a/Engine/Source/Runtime/TypedElementFramework/Public/Elements/Interfaces/TypedElementQueryStorageInterfaces.h b/Engine/Source/Runtime/TypedElementFramework/Public/Elements/Interfaces/TypedElementQueryStorageInterfaces.h
--- a/Engine/Source/Runtime/TypedElementFramework/Public/Elements/Interfaces/TypedElementQueryStorageInterfaces.h
+++ b/Engine/Source/Runtime/TypedElementFramework/Public/Elements/Interfaces/TypedElementQueryStorageInterfaces.h
@@ -40,8 +40,15 @@
 
 		/** Return the address of a immutable column matching the requested type or a nullptr if not found. */
 		virtual const void* GetColumn(const UScriptStruct* ColumnType) const = 0;
+		/** Return the address of a immutable column matching the requested type or a nullptr if not found. */
+		template<typename Column>
+		const Column* GetColumn() const;
 		/** Return the address of a mutable column matching the requested type or a nullptr if not found. */
 		virtual void* GetMutableColumn(const UScriptStruct* ColumnType) = 0;
+		/** Return the address of a mutable column matching the requested type or a nullptr if not found. */
+		template<typename Column>
+		Column* GetMutableColumn();
+
 		/**
 		 * Get a list of columns or nullptrs if the column type wasn't found. Mutable addresses are returned and it's up to
 		 * the caller to not change immutable addresses.
@@ -55,73 +62,48 @@
 		 */
 		virtual void GetColumnsUnguarded(int32 TypeCount, char** RetrievedAddresses, const TWeakObjectPtr<const UScriptStruct>* ColumnTypes,
 			const EQueryAccessType* AccessTypes) = 0;
-		/* Return whether a column matches the requested type or not. */
+		/* 
+		 * Returns whether a column matches the requested type or not. This version only applies to the table that's currently set in
+		 * the context. This version is faster of checking in the current row or table, but the version using a row is needed to check
+		 * arbitrary rows.
+		 */
 		virtual bool HasColumn(const UScriptStruct* ColumnType) const = 0;
-		// Utility functions
-
+		/*
+		 * Returns whether a column matches the requested type or not. This version only applies to the table that's currently set in
+		 * the context. This version is faster of checking in the current row or table, but the version using a row is needed to check
+		 * arbitrary rows.
+		 */
 		template<typename Column>
-		const Column* GetColumn() const;
-		template<typename Column>
-		Column* GetMutableColumn();
-		template<typename Column>
 		bool HasColumn() const;
+		/*
+		 * Return whether a column matches the requested type or not. This can be used for arbitrary rows. If the row is in the
+		 * table that's set in the context, for instance because it's the current row, then the version that doesn't take a row
+		 * as an argument is recommended.
+		 */
+		virtual bool HasColumn(TypedElementDataStorage::RowHandle Row, const UScriptStruct* ColumnType) const = 0;
+		/*
+		 * Return whether a column matches the requested type or not. This can be used for arbitrary rows. If the row is in the
+		 * table that's set in the context, for instance because it's the current row, then the version that doesn't take a row
+		 * as an argument is recommended.
+		 */
+		template<typename Column>
+		bool HasColumn(TypedElementDataStorage::RowHandle Row) const;
 	};
 
-	/**
-	 * Interface to be provided to query callbacks that are directly called through RunQuery from outside a query callback.
-	 */
-	struct IDirectQueryContext : public ICommonQueryContext
+	struct ICommonQueryWithEnvironmentContext : public ICommonQueryContext
 	{
-		virtual ~IDirectQueryContext() = default;
-	};
-
-	/**
-	 * Interface to be provided to query callbacks that are directly called through from a query callback.
-	 */
-	struct ISubqueryContext : public ICommonQueryContext
-	{
-		virtual ~ISubqueryContext() = default;
-	};
-
-	/**
-	 * Interface to be provided to query callbacks running with the Data Storage.
-	 * Note that at the time of writing only subclasses of Subsystem are supported as dependencies.
-	 */
-	struct IQueryContext : public ICommonQueryContext
-	{
 		using ObjectCopyOrMove = void (*)(const UScriptStruct* TypeInfo, void* Destination, void* Source);
 
-		virtual ~IQueryContext() = default;
+		/**
+		 * Returns the id for the current update cycle. Every time TEDS goes through a cycle of running query callbacks, this is
+		 * incremented by one. This guarantees that all query callbacks in the same run see the same cycle id. This can be useful
+		 * in avoiding duplicated work.
+		 */
+		virtual uint64 GetUpdateCycleId() const = 0;
 
-		/** Returns an immutable instance of the requested dependency or a nullptr if not found. */
-		virtual const UObject* GetDependency(const UClass* DependencyClass) = 0;
-		/** Returns a mutable instance of the requested dependency or a nullptr if not found. */
-		virtual UObject* GetMutableDependency(const UClass* DependencyClass) = 0;
 		/**
-		 * Returns a list of dependencies or nullptrs if a dependency wasn't found. Mutable versions are return and it's up to the
-		 * caller to not change immutable dependencies.
-		 */
-		virtual void GetDependencies(TArrayView<UObject*> RetrievedAddresses, TConstArrayView<TWeakObjectPtr<const UClass>> DependencyTypes,
-			TConstArrayView<EQueryAccessType> AccessTypes) = 0;
-
-		/** Checks whether or not a row is in use. This is true even if the row has only been reserved. */
-		virtual bool IsRowAvailable(RowHandle Row) const = 0;
-		/** Checks whether or not a row has been reserved but not yet assigned to a table. */
-		virtual bool HasRowBeenAssigned(RowHandle Row) const = 0;
-		/**
-		 * Removes the row with the provided row handle. The removal will not be immediately done but delayed until the end of the tick
-		 * group.
-		 */
-		virtual void RemoveRow(RowHandle Row) = 0;
-		/**
-		 * Removes rows with the provided row handles. The removal will not be immediately done but delayed until the end of the tick
-		 * group.
-		 */
-		virtual void RemoveRows(TConstArrayView<RowHandle> Rows) = 0;
-		
-		/**
 		 * Adds the provided column to the requested row.
-		 * 
+		 *
 		 * Note: The addition of the column will not be immediately done. Instead it will be deferred until the end of the tick group. Changes
 		 * made to the return column will still be applied when the column is added to the row.
 		 */
@@ -153,7 +135,7 @@
 		 * Add a new uninitialized column of the provided type if one does not exist.
 		 * Returns a staged column which is used to copy into the database at a later time via the UStructScript Copy operator at the end
 		 * of the tick group.
-		 * It is the caller's responsibility to ensure the staged column's constructor is called. The caller may modify other 
+		 * It is the caller's responsibility to ensure the staged column's constructor is called. The caller may modify other
 		 * values in the column.
 		 * This function can not be used to add a tag as tags do not contain any data.
 		 */
@@ -162,7 +144,7 @@
 		 * Add a new uninitialized column of the provided type if one does not exist.
 		 * Returns a staged column which is used to copy/move into the database at a later time via the provided relocator at the end of
 		 * the tick group.
-		 * It is the caller's responsibility to ensure the staged column's constructor is called. The caller may modify other 
+		 * It is the caller's responsibility to ensure the staged column's constructor is called. The caller may modify other
 		 * values in the column.
 		 * This function can not be used to add a tag as tags do not contain any data.
 		 */
@@ -190,6 +172,57 @@
 		 * tick group.
 		 */
 		virtual void RemoveColumns(TConstArrayView<RowHandle> Rows, TConstArrayView<const UScriptStruct*> ColumnTypes) = 0;
+	};
 
+	/**
+	 * Interface to be provided to query callbacks that are directly called through RunQuery from outside a query callback.
+	 */
+	struct IDirectQueryContext : public ICommonQueryContext
+	{
+		virtual ~IDirectQueryContext() = default;
+	};
+
+	/**
+	 * Interface to be provided to query callbacks that are directly called through from a query callback.
+	 */
+	struct ISubqueryContext : public ICommonQueryWithEnvironmentContext
+	{
+		virtual ~ISubqueryContext() = default;
+	};
+
+	/**
+	 * Interface to be provided to query callbacks running with the Data Storage.
+	 * Note that at the time of writing only subclasses of Subsystem are supported as dependencies.
+	 */
+	struct IQueryContext : public ICommonQueryWithEnvironmentContext
+	{
+		virtual ~IQueryContext() = default;
+
+		/** Returns an immutable instance of the requested dependency or a nullptr if not found. */
+		virtual const UObject* GetDependency(const UClass* DependencyClass) = 0;
+		/** Returns a mutable instance of the requested dependency or a nullptr if not found. */
+		virtual UObject* GetMutableDependency(const UClass* DependencyClass) = 0;
+		/**
+		 * Returns a list of dependencies or nullptrs if a dependency wasn't found. Mutable versions are return and it's up to the
+		 * caller to not change immutable dependencies.
+		 */
+		virtual void GetDependencies(TArrayView<UObject*> RetrievedAddresses, TConstArrayView<TWeakObjectPtr<const UClass>> DependencyTypes,
+			TConstArrayView<EQueryAccessType> AccessTypes) = 0;
+
+		/** Checks whether or not a row is in use. This is true even if the row has only been reserved. */
+		virtual bool IsRowAvailable(RowHandle Row) const = 0;
+		/** Checks whether or not a row has been reserved but not yet assigned to a table. */
+		virtual bool HasRowBeenAssigned(RowHandle Row) const = 0;
+		/**
+		 * Removes the row with the provided row handle. The removal will not be immediately done but delayed until the end of the tick
+		 * group.
+		 */
+		virtual void RemoveRow(RowHandle Row) = 0;
+		/**
+		 * Removes rows with the provided row handles. The removal will not be immediately done but delayed until the end of the tick
+		 * group.
+		 */
+		virtual void RemoveRows(TConstArrayView<RowHandle> Rows) = 0;
+
 		/** Retrieves the row for an indexed object. Returns an invalid row handle if the hash wasn't found. */
 		virtual RowHandle FindIndexedRow(IndexHash Index) const = 0;
@@ -229,6 +262,10 @@
 
 namespace TypedElementDataStorage
 {
+	//
+	// ICommonQueryContext
+	//
+
 	template<typename Column>
 	const Column* ICommonQueryContext::GetColumn() const
 	{
@@ -247,8 +284,20 @@
 		return HasColumn(Column::StaticStruct());
 	}
 
+	template <typename Column>
+	bool ICommonQueryContext::HasColumn(TypedElementDataStorage::RowHandle Row) const
+	{
+		return HasColumn(Row, Column::StaticStruct());
+	}
+
+
+
+	//
+	// ICommonQueryWithEnvironmentContext
+	// 
+
 	template<typename ColumnType>
-	ColumnType& IQueryContext::AddColumn(RowHandle Row, ColumnType&& Column)
+	ColumnType& ICommonQueryWithEnvironmentContext::AddColumn(RowHandle Row, ColumnType&& Column)
 	{
 		UScriptStruct* TypeInfo = ColumnType::StaticStruct();
 
@@ -270,25 +319,25 @@
 	}
 
 	template<typename... Columns>
-	void IQueryContext::AddColumns(RowHandle Row)
+	void ICommonQueryWithEnvironmentContext::AddColumns(RowHandle Row)
 	{
 		AddColumns(Row, { Columns::StaticStruct()... });
 	}
 
 	template<typename... Columns>
-	void IQueryContext::AddColumns(TConstArrayView<RowHandle> Rows)
+	void ICommonQueryWithEnvironmentContext::AddColumns(TConstArrayView<RowHandle> Rows)
 	{
 		AddColumns(Rows, { Columns::StaticStruct()... });
 	}
 
 	template<typename... Columns>
-	void IQueryContext::RemoveColumns(RowHandle Row)
+	void ICommonQueryWithEnvironmentContext::RemoveColumns(RowHandle Row)
 	{
 		RemoveColumns(Row, { Columns::StaticStruct()... });
 	}
 
 	template<typename... Columns>
-	void IQueryContext::RemoveColumns(TConstArrayView<RowHandle> Rows)
+	void ICommonQueryWithEnvironmentContext::RemoveColumns(TConstArrayView<RowHandle> Rows)
 	{
 		RemoveColumns(Rows, { Columns::StaticStruct()... });
 	}
