Fixed Mac compile error.

diff --git a/Engine/Plugins/Experimental/TypedElementsDataStorage/Source/TypedElementsDataStorage/Private/Processors/TypedElementProcessorAdaptors.cpp b/Engine/Plugins/Experimental/TypedElementsDataStorage/Source/TypedElementsDataStorage/Private/Processors/TypedElementProcessorAdaptors.cpp
--- a/Engine/Plugins/Experimental/TypedElementsDataStorage/Source/TypedElementsDataStorage/Private/Processors/TypedElementProcessorAdaptors.cpp
+++ b/Engine/Plugins/Experimental/TypedElementsDataStorage/Source/TypedElementsDataStorage/Private/Processors/TypedElementProcessorAdaptors.cpp
@@ -205,7 +205,7 @@ public:
 		struct FAddMoveableValueColumnWithDestructor : FAddMoveableValueColumn
 		{
 			FAddMoveableValueColumnWithDestructor(
-				Parent::ObjectCopyOrMove InRelocator, const UScriptStruct* InFragmentType, FMassEntityHandle InEntity, void* InObject)
+				typename Parent::ObjectCopyOrMove InRelocator, const UScriptStruct* InFragmentType, FMassEntityHandle InEntity, void* InObject)
 				: FAddMoveableValueColumn(InRelocator, InFragmentType, InEntity, InObject)
 			{}
 
@@ -265,7 +265,7 @@ public:
 			FMassEntityHandle Entity;
 		};
 
-		FAddedColumns* AddedColumns = Environment.GetScratchBuffer().Emplace<FAddedColumns>();
+		FAddedColumns* AddedColumns = Environment.GetScratchBuffer().template Emplace<FAddedColumns>();
 		TedsColumnsToMassDescriptor(AddedColumns->AddDescriptor, ColumnTypes);
 		AddedColumns->Entity = FMassEntityHandle::FromNumber(Row);
 
@@ -324,7 +324,7 @@ public:
 			FMassEntityHandle Entity;
 		};
 
-		FRemovedColumns* RemovedColumns = Environment.GetScratchBuffer().Emplace<FRemovedColumns>();
+		FRemovedColumns* RemovedColumns = Environment.GetScratchBuffer().template Emplace<FRemovedColumns>();
 		TedsColumnsToMassDescriptorIfActiveTable(RemovedColumns->RemoveDescriptor, ColumnTypes);
 		if (!RemovedColumns->RemoveDescriptor.IsEmpty())
 		{
