Fixed a Clang Win64 compiler error.

diff --git a/Engine/Plugins/Experimental/TypedElementsDataStorage/Source/TypedElementsDataStorage/Private/Processors/TypedElementProcessorAdaptors.cpp b/Engine/Plugins/Experimental/TypedElementsDataStorage/Source/TypedElementsDataStorage/Private/Processors/TypedElementProcessorAdaptors.cpp
--- a/Engine/Plugins/Experimental/TypedElementsDataStorage/Source/TypedElementsDataStorage/Private/Processors/TypedElementProcessorAdaptors.cpp
+++ b/Engine/Plugins/Experimental/TypedElementsDataStorage/Source/TypedElementsDataStorage/Private/Processors/TypedElementProcessorAdaptors.cpp
@@ -186,47 +186,35 @@ public:
 	{
 		checkf(ObjectType->IsChildOf(FMassFragment::StaticStruct()), TEXT("Column [%s] can not be a tag"), *ObjectType->GetName());
 		
-		struct FAddMoveableValueColumn
+		struct FAddValueColumn
 		{
 			typename Parent::ObjectCopyOrMove Relocator;
 			const UScriptStruct* FragmentType;
 			FMassEntityHandle Entity;
 			void* Object;
 
-			FAddMoveableValueColumn() = default;
-			FAddMoveableValueColumn(typename Parent::ObjectCopyOrMove InRelocator, const UScriptStruct* InFragmentType, FMassEntityHandle InEntity, void* InObject)
+			FAddValueColumn() = default;
+			FAddValueColumn(typename Parent::ObjectCopyOrMove InRelocator, const UScriptStruct* InFragmentType, FMassEntityHandle InEntity, void* InObject)
 				: Relocator(InRelocator)
 				, FragmentType(InFragmentType)
 				, Entity(InEntity)
 				, Object(InObject)
 			{}
-		};
-
-		struct FAddMoveableValueColumnWithDestructor : FAddMoveableValueColumn
-		{
-			FAddMoveableValueColumnWithDestructor(
-				typename Parent::ObjectCopyOrMove InRelocator, const UScriptStruct* InFragmentType, FMassEntityHandle InEntity, void* InObject)
-				: FAddMoveableValueColumn(InRelocator, InFragmentType, InEntity, InObject)
-			{}
 
-			~FAddMoveableValueColumnWithDestructor()
+			~FAddValueColumn()
 			{
-				this->FragmentType->DestroyStruct(this->Object);
+				if ((this->FragmentType->StructFlags & (STRUCT_IsPlainOldData | STRUCT_NoDestructor)) == 0)
+				{
+					this->FragmentType->DestroyStruct(this->Object);
+				}
 			}
 		};
 
 		FTypedElementDatabaseScratchBuffer& ScratchBuffer = Environment.GetScratchBuffer();
 		void* ColumnData = ScratchBuffer.Allocate(ObjectType->GetStructureSize(), ObjectType->GetMinAlignment());
-		FAddMoveableValueColumn* AddedColumn = nullptr;
-		if (ObjectType->StructFlags & (STRUCT_IsPlainOldData | STRUCT_NoDestructor))
-		{
-			AddedColumn = ScratchBuffer.Emplace<FAddMoveableValueColumn>(Relocator, ObjectType, FMassEntityHandle::FromNumber(Row), ColumnData);
-		}
-		else
-		{
-			AddedColumn = ScratchBuffer.Emplace<FAddMoveableValueColumnWithDestructor>(Relocator, ObjectType, FMassEntityHandle::FromNumber(Row), ColumnData);
-		}
-
+		FAddValueColumn* AddedColumn =
+			ScratchBuffer.Emplace<FAddValueColumn>(Relocator, ObjectType, FMassEntityHandle::FromNumber(Row), ColumnData);
+		
 		this->Context.Defer().template PushCommand<FMassDeferredAddCommand>(
 			[AddedColumn](FMassEntityManager& System)
 			{
