Add editor-only color overlays for objects in viewport

#p4v-preflight-copy 31541584

diff --git a/Engine/Plugins/Experimental/TypedElementsDataStorage/Source/TypedElementsDataStorage/Private/Compatibility/TypedElementActorViewportProcessors.cpp b/Engine/Plugins/Experimental/TypedElementsDataStorage/Source/TypedElementsDataStorage/Private/Compatibility/TypedElementActorViewportProcessors.cpp
--- a/Engine/Plugins/Experimental/TypedElementsDataStorage/Source/TypedElementsDataStorage/Private/Compatibility/TypedElementActorViewportProcessors.cpp
+++ b/Engine/Plugins/Experimental/TypedElementsDataStorage/Source/TypedElementsDataStorage/Private/Compatibility/TypedElementActorViewportProcessors.cpp
@@ -6,24 +6,139 @@
 #include "Elements/Columns/TypedElementViewportColumns.h"
 #include "Elements/Columns/TypedElementMiscColumns.h"
 #include "Elements/Framework/TypedElementQueryBuilder.h"
+#include "Elements/Framework/TypedElementRegistry.h"
+#include "HAL/IConsoleManager.h"
 #include "MassActorSubsystem.h"
+#include "Elements/Columns/TypedElementSelectionColumns.h"
+#include "ProfilingDebugging/CpuProfilerTrace.h"
+
+FAutoConsoleCommandWithArgsAndOutputDevice SetOutlineColorConsoleCommand(
+	TEXT("TEDS.Debug.SetOutlineColor"),
+	TEXT("Adds an outline color to selected objects."),
+	FConsoleCommandWithArgsAndOutputDeviceDelegate::CreateLambda([](const TArray<FString>& Args, FOutputDevice& Output)
+		{
+			using namespace TypedElementQueryBuilder;
+			using DSI = ITypedElementDataStorageInterface;
+
+			TRACE_CPUPROFILER_EVENT_SCOPE(TEDS.Debug.AddOverlayColorToSelectionCommand);
+
+			if (ITypedElementDataStorageInterface* DataStorage = UTypedElementRegistry::GetInstance()->GetMutableDataStorage())
+			{
+				static TypedElementQueryHandle OverlayQuery = TypedElementInvalidQueryHandle;
+				if (OverlayQuery == TypedElementInvalidQueryHandle)
+				{
+					OverlayQuery = DataStorage->RegisterQuery(
+						Select()
+						.Where()
+							.All<FTypedElementSelectionColumn>()
+						.Compile());
+				}
+				
+				if (OverlayQuery == TypedElementInvalidQueryHandle)
+				{
+					return;
+				}
+
+				if (Args.IsEmpty())
+				{
+					Output.Log(TEXT("Provide a color index (0-7) to use as outline"));
+					return;
+				}
+				
+				// Parse the color
+				int32 ColorIndex;
+				LexFromString(ColorIndex, *Args[0]);
+
+				if (!(ColorIndex >= 0 && ColorIndex <= 7))
+				{
+					Output.Log(TEXT("Color index must be in range [0,7]"));
+					return;
+				}
+
+				TArray<TypedElementRowHandle> RowHandles;
+				
+				DataStorage->RunQuery(OverlayQuery, [ColorIndex, &RowHandles](const DSI::FQueryDescription&, DSI::IDirectQueryContext& Context)
+				{
+					RowHandles = Context.GetRowHandles();
+				});
+
+				for (TypedElementRowHandle Row : RowHandles)
+				{
+					DataStorage->AddOrGetColumn<FTypedElementViewportOutlineColorColumn>(Row)->SelectionOutlineColorIndex = ColorIndex;
+					DataStorage->AddColumn<FTypedElementSyncBackToWorldTag>(Row);
+				}
+			}
+		}));
+
+FAutoConsoleCommandWithArgsAndOutputDevice SetSelectionOverlayColorConsoleCommand(
+	TEXT("TEDS.Debug.SetOverlayColor"),
+	TEXT("Adds an overlay color to selected objects."),
+	FConsoleCommandWithArgsAndOutputDeviceDelegate::CreateLambda([](const TArray<FString>& Args, FOutputDevice& Output)
+		{
+			using namespace TypedElementQueryBuilder;
+			using DSI = ITypedElementDataStorageInterface;
+
+			TRACE_CPUPROFILER_EVENT_SCOPE(TEDS.Debug.AddOverlayColorToSelectionCommand);
+
+			if (ITypedElementDataStorageInterface* DataStorage = UTypedElementRegistry::GetInstance()->GetMutableDataStorage())
+			{
+				static TypedElementQueryHandle OverlayQuery = TypedElementInvalidQueryHandle;
+				if (OverlayQuery == TypedElementInvalidQueryHandle)
+				{
+					OverlayQuery = DataStorage->RegisterQuery(
+						Select()
+						.Where()
+							.All<FTypedElementSelectionColumn>()
+						.Compile());
+				}
+				
+				if (OverlayQuery == TypedElementInvalidQueryHandle)
+				{
+					return;
+				}
+
+				if (Args.IsEmpty())
+				{
+					Output.Log(TEXT("Provide a color in hexadecimal format (#RRGGBBAA) to overlay."));
+					return;
+				}
+				
+				// Parse the color
+				FColor Color = FColor::FromHex(Args[0]);
+				Color.A = FMath::Clamp(Color.A, 0, 128);
+
+				TArray<TypedElementRowHandle> RowHandles;
+				
+				DataStorage->RunQuery(OverlayQuery, [Color, &RowHandles](const DSI::FQueryDescription&, DSI::IDirectQueryContext& Context)
+				{
+					RowHandles = Context.GetRowHandles();
+				});
+
+				for (TypedElementRowHandle Row : RowHandles)
+				{
+					DataStorage->AddOrGetColumn<FTypedElementViewportOverlayColorColumn>(Row)->OverlayColor = Color;
+					DataStorage->AddColumn<FTypedElementSyncBackToWorldTag>(Row);
+				}
+			}
+		}));
 
 void UTypedElementActorViewportFactory::RegisterQueries(ITypedElementDataStorageInterface& DataStorage)
 {
-	RegisterSelectionOutlineColorColumnToActor(DataStorage);
+	RegisterOutlineColorColumnToActor(DataStorage);
+	RegisterOverlayColorColumnToActor(DataStorage);
 }
 
-void UTypedElementActorViewportFactory::RegisterSelectionOutlineColorColumnToActor(ITypedElementDataStorageInterface& DataStorage) const
+void UTypedElementActorViewportFactory::RegisterOutlineColorColumnToActor(ITypedElementDataStorageInterface& DataStorage)
 {
 	using namespace TypedElementQueryBuilder;
 	using DSI = ITypedElementDataStorageInterface;
 
 	DataStorage.RegisterQuery(
 		Select(
-			TEXT("Sync selection outline color column to actor"),
+			TEXT("Sync viewport outline color column to actor"),
 			FProcessor(DSI::EQueryTickPhase::DuringPhysics, DataStorage.GetQueryTickGroupName(DSI::EQueryTickGroups::SyncDataStorageToExternal))
 				.ForceToGameThread(true),
-			[](FMassActorFragment& Actor, const FTypedElementViewportColorColumn& ViewportColor)
+			[](FMassActorFragment& Actor, const FTypedElementViewportOutlineColorColumn& ViewportColor)
 			{
 				if (AActor* ActorInstance = Actor.GetMutable(); ActorInstance != nullptr)
 				{
@@ -39,3 +154,30 @@ void UTypedElementActorViewportFactory::RegisterSelectionOutlineColorColumnToAct
 		.Compile()
 	);
 }
+
+void UTypedElementActorViewportFactory::RegisterOverlayColorColumnToActor(ITypedElementDataStorageInterface& DataStorage)
+{
+	using namespace TypedElementQueryBuilder;
+	using DSI = ITypedElementDataStorageInterface;
+
+	DataStorage.RegisterQuery(
+		Select(
+			TEXT("Sync viewport overlay color column to actor"),
+			FProcessor(DSI::EQueryTickPhase::DuringPhysics, DataStorage.GetQueryTickGroupName(DSI::EQueryTickGroups::SyncDataStorageToExternal))
+				.ForceToGameThread(true),
+			[](FMassActorFragment& Actor, const FTypedElementViewportOverlayColorColumn& ViewportColor)
+			{
+				if (AActor* ActorInstance = Actor.GetMutable(); ActorInstance != nullptr)
+				{
+					if (UPrimitiveComponent* PrimitiveComponent = ActorInstance->GetComponentByClass<UPrimitiveComponent>())
+					{
+						PrimitiveComponent->SetOverlayColor(ViewportColor.OverlayColor);
+					}
+				}
+			}
+		)
+		.Where()
+			.All<FTypedElementSyncBackToWorldTag>()
+		.Compile()
+	);
+}
diff --git a/Engine/Plugins/Experimental/TypedElementsDataStorage/Source/TypedElementsDataStorage/Private/Compatibility/TypedElementActorViewportProcessors.h b/Engine/Plugins/Experimental/TypedElementsDataStorage/Source/TypedElementsDataStorage/Private/Compatibility/TypedElementActorViewportProcessors.h
--- a/Engine/Plugins/Experimental/TypedElementsDataStorage/Source/TypedElementsDataStorage/Private/Compatibility/TypedElementActorViewportProcessors.h
+++ b/Engine/Plugins/Experimental/TypedElementsDataStorage/Source/TypedElementsDataStorage/Private/Compatibility/TypedElementActorViewportProcessors.h
@@ -18,5 +18,6 @@ public:
 	void RegisterQueries(ITypedElementDataStorageInterface& DataStorage) override;
 
 private:
-	void RegisterSelectionOutlineColorColumnToActor(ITypedElementDataStorageInterface& DataStorage) const;
+	void RegisterOutlineColorColumnToActor(ITypedElementDataStorageInterface& DataStorage);
+	void RegisterOverlayColorColumnToActor(ITypedElementDataStorageInterface& DataStorage);
 };
diff --git a/Engine/Plugins/Experimental/TypedElementsDataStorage/Source/TypedElementsDataStorageRevisionControl/Source/Private/RevisionControlProcessors.cpp b/Engine/Plugins/Experimental/TypedElementsDataStorage/Source/TypedElementsDataStorageRevisionControl/Source/Private/RevisionControlProcessors.cpp
--- a/Engine/Plugins/Experimental/TypedElementsDataStorage/Source/TypedElementsDataStorageRevisionControl/Source/Private/RevisionControlProcessors.cpp
+++ b/Engine/Plugins/Experimental/TypedElementsDataStorage/Source/TypedElementsDataStorageRevisionControl/Source/Private/RevisionControlProcessors.cpp
@@ -99,7 +99,7 @@ void UTypedElementRevisionControlFactory::RegisterQueries(ITypedElementDataStora
 				Context.RunSubquery(0, PackageReference.Row, CreateSubqueryCallbackBinding(
 					[&Context, &ObjectRow](DSI::ISubqueryContext& SubQueryContext)
 					{
-						Context.AddColumn<FTypedElementViewportColorColumn>(ObjectRow, { .SelectionOutlineColorIndex = DetermineOutlineColorIndex(SubQueryContext) });
+						Context.AddColumn<FTypedElementViewportOutlineColorColumn>(ObjectRow, { .SelectionOutlineColorIndex = DetermineOutlineColorIndex(SubQueryContext) });
 						Context.AddColumns<FTypedElementSyncBackToWorldTag>(ObjectRow);
 					})
 				);
diff --git a/Engine/Shaders/Private/PostProcessSelectionOutline.usf b/Engine/Shaders/Private/PostProcessSelectionOutline.usf
--- a/Engine/Shaders/Private/PostProcessSelectionOutline.usf
+++ b/Engine/Shaders/Private/PostProcessSelectionOutline.usf
@@ -20,6 +20,9 @@ SamplerState ColorSampler;
 Texture2D DepthTexture;
 SamplerState DepthSampler;
 
+Texture2D OverlayLookupTexture;
+SamplerState OverlayLookupSampler;
+
 SCREEN_PASS_TEXTURE_VIEWPORT(Color)
 SCREEN_PASS_TEXTURE_VIEWPORT(Depth)
 FScreenTransform ColorToDepth;
@@ -233,32 +236,36 @@ void MainPS(
 	);
 #endif
 
-	const int ScoutSum = dot(saturate(StencilScout), 1);
+	const float2 DepthUV = ApplyScreenTransform(ColorUV, ColorToDepth);
 
-	// If this sum is zero, none of our neighbors are selected pixels and we can skip all the heavy processing.
-	if (ScoutSum > 0)
-	{
-		const float2 DepthUV = ApplyScreenTransform(ColorUV, ColorToDepth);
+	float Center = Texture2DSampleLevel(DepthTexture, DepthSampler, DepthUV, 0).r;
+	float Left   = Texture2DSampleLevel(DepthTexture, DepthSampler, DepthUV + float2(-1,  0) * Depth_ExtentInverse, 0).r;
+	float Right  = Texture2DSampleLevel(DepthTexture, DepthSampler, DepthUV + float2( 1,  0) * Depth_ExtentInverse, 0).r;
+	float Top    = Texture2DSampleLevel(DepthTexture, DepthSampler, DepthUV + float2( 0, -1) * Depth_ExtentInverse, 0).r;
+	float Bottom = Texture2DSampleLevel(DepthTexture, DepthSampler, DepthUV + float2( 0,  1) * Depth_ExtentInverse, 0).r;
 
-		float Center = Texture2DSampleLevel(DepthTexture, DepthSampler, DepthUV, 0).r;
-		float Left   = Texture2DSampleLevel(DepthTexture, DepthSampler, DepthUV + float2(-1,  0) * Depth_ExtentInverse, 0).r;
-		float Right  = Texture2DSampleLevel(DepthTexture, DepthSampler, DepthUV + float2( 1,  0) * Depth_ExtentInverse, 0).r;
-		float Top    = Texture2DSampleLevel(DepthTexture, DepthSampler, DepthUV + float2( 0, -1) * Depth_ExtentInverse, 0).r;
-		float Bottom = Texture2DSampleLevel(DepthTexture, DepthSampler, DepthUV + float2( 0,  1) * Depth_ExtentInverse, 0).r;
+	// This allows to reconstruct depth with a small pixel offset without many texture lookups (4xMSAA * 5 neighbors -> 20 samples)
+	// It's an approximation assuming the surface is a plane.
+	float3 DeviceZPlane;
+	DeviceZPlane.x = (Right - Left) / 2;
+	DeviceZPlane.y = (Bottom - Top) / 2;
+	DeviceZPlane.z = Center;
 
-		// This allows to reconstruct depth with a small pixel offset without many texture lookups (4xMSAA * 5 neighbors -> 20 samples)
-		// It's an approximation assuming the surface is a plane.
-		float3 DeviceZPlane;
-		DeviceZPlane.x = (Right - Left) / 2;
-		DeviceZPlane.y = (Bottom - Top) / 2;
-		DeviceZPlane.z = Center;
+	float2 DeviceZMinMax;
+	DeviceZMinMax.x = min(Center, min(min(Left, Right), min(Top, Bottom)));
+	DeviceZMinMax.y = max(Center, max(max(Left, Right), max(Top, Bottom)));
 
-		float2 DeviceZMinMax;
-		DeviceZMinMax.x = min(Center, min(min(Left, Right), min(Top, Bottom)));
-		DeviceZMinMax.y = max(Center, max(max(Left, Right), max(Top, Bottom)));
+	FPixelInfo CenterPixelInfo = GetPixelInfo(ColorPixelPos, 0, SecondaryViewportOffset, 0, DeviceZPlane, DeviceZMinMax);
 
-		FPixelInfo CenterPixelInfo = GetPixelInfo(ColorPixelPos, 0, SecondaryViewportOffset, 0, DeviceZPlane, DeviceZMinMax);
+	float4 OverlayColor = Texture2DSample(OverlayLookupTexture, OverlayLookupSampler, ColorUV);
+	OutColor = Texture2DSample(ColorTexture, ColorSampler, ColorUV); 
+	
+	float4 SelectionEffect;
+	const int ScoutSum = dot(saturate(StencilScout), 1);
 
+	// If this sum is zero, none of our neighbors are selected pixels and we can skip all the heavy processing.
+	if (ScoutSum > 0)
+	{
 		float4 Sum = 0;
 #if COMPILER_GLSL && FEATURE_LEVEL >= FEATURE_LEVEL_SM4
 	UNROLL
@@ -267,49 +274,47 @@ void MainPS(
 		{
 			Sum += PerSample(ColorPixelPos, SampleID, CenterPixelInfo, DeviceZPlane, DeviceZMinMax);
 		}
-		float4 Avg = Sum / MSAA_SAMPLE_COUNT;
-
-		OutColor = Texture2DSample(ColorTexture, ColorSampler, ColorUV);
-
-		// Scene color has gamma applied. Need to remove and re-apply after linear operation.
-		if (Avg.a > 0)
-		{
+		SelectionEffect = Sum / MSAA_SAMPLE_COUNT;
+	}
+	
+	// Scene color has gamma applied. Need to remove and re-apply after linear operations.
+	if (SelectionEffect.a > 0 || OverlayColor.a > 0)
+	{
 #if SELECTION_OUTLINE_HDR
-			static const float ScRGBScaleFactor = 80.0f; // used to convert between ScRGB and nits
-			Avg.rgb *= abs(UILuminanceAndIsSCRGB);
-			if (UILuminanceAndIsSCRGB > 0.0)
-			{
-				OutColor.rgb *= ScRGBScaleFactor;
-			}
-			else
-			{
-				const float3x3 Rec2020_2_sRGB = mul(XYZ_2_sRGB_MAT, Rec2020_2_XYZ_MAT);
-				OutColor.rgb = ST2084ToLinear(OutColor.rgb);
-				OutColor.rgb = mul(Rec2020_2_sRGB, OutColor.rgb);
-			}
+		static const float ScRGBScaleFactor = 80.0f; // used to convert between ScRGB and nits
+		SelectionEffect.rgb *= abs(UILuminanceAndIsSCRGB);
+		OverlayColor.rgb *= abs(UILuminanceAndIsSCRGB);
+		if (UILuminanceAndIsSCRGB > 0.0)
+		{
+			OutColor.rgb *= ScRGBScaleFactor;
+		}
+		else
+		{
+			const float3x3 Rec2020_2_sRGB = mul(XYZ_2_sRGB_MAT, Rec2020_2_XYZ_MAT);
+			OutColor.rgb = ST2084ToLinear(OutColor.rgb);
+			OutColor.rgb = mul(Rec2020_2_sRGB, OutColor.rgb);
+		}
 #else
-			OutColor.rgb = sRGBToLinear(OutColor.rgb);
+		OutColor.rgb = sRGBToLinear(OutColor.rgb);
 #endif
-			OutColor.rgb = lerp(OutColor.rgb, Avg.rgb, Avg.a);
+		float PatternMask = ((ColorPixelPos.x/2 + ColorPixelPos.y/2) % 2) * 0.6f;
+		float VisibleMask = lerp(PatternMask, 1.0f, CenterPixelInfo.fVisible);
+		OutColor.rgb = lerp(OutColor.rgb, OverlayColor.rgb, OverlayColor.a * VisibleMask);
+		OutColor.rgb = lerp(OutColor.rgb, SelectionEffect.rgb, SelectionEffect.a);
 
 #if SELECTION_OUTLINE_HDR
-			if (UILuminanceAndIsSCRGB > 0.0)
-			{
-				OutColor.rgb /= ScRGBScaleFactor;
-			}
-			else
-			{
-				const float3x3 sRGB_2_Rec2020 = mul(XYZ_2_Rec2020_MAT, sRGB_2_XYZ_MAT);
-				OutColor.xyz = mul(sRGB_2_Rec2020, OutColor.xyz);
-				OutColor.rgb = LinearToST2084(OutColor.rgb);
-			}
+		if (UILuminanceAndIsSCRGB > 0.0)
+		{
+			OutColor.rgb /= ScRGBScaleFactor;
+		}
+		else
+		{
+			const float3x3 sRGB_2_Rec2020 = mul(XYZ_2_Rec2020_MAT, sRGB_2_XYZ_MAT);
+			OutColor.xyz = mul(sRGB_2_Rec2020, OutColor.xyz);
+			OutColor.rgb = LinearToST2084(OutColor.rgb);
+		}
 #else
-			OutColor.rgb = LinearToSrgbBranchless(OutColor.rgb);
+		OutColor.rgb = LinearToSrgbBranchless(OutColor.rgb);
 #endif
-		}
-	}
-	else
-	{
-		OutColor = Texture2DSample(ColorTexture, ColorSampler, ColorUV);
 	}
 }
diff --git a/Engine/Source/Runtime/Engine/Classes/Components/PrimitiveComponent.h b/Engine/Source/Runtime/Engine/Classes/Components/PrimitiveComponent.h
--- a/Engine/Source/Runtime/Engine/Classes/Components/PrimitiveComponent.h
+++ b/Engine/Source/Runtime/Engine/Classes/Components/PrimitiveComponent.h
@@ -443,6 +443,10 @@ public:
 	UPROPERTY()
 	uint8 bSelectable:1;
 
+	/** When true, this component requests editor effects like outlines and overlays. */
+	UPROPERTY()
+	uint8 bWantsEditorEffects:1;
+
 #if WITH_EDITORONLY_DATA
 	/** If true, this component will be considered for placement when dragging and placing items in the editor even if it is not visible, such as in the case of hidden collision meshes */
 	UPROPERTY(EditAnywhere, AdvancedDisplay, Category = Collision)
@@ -673,6 +677,8 @@ public:
 	uint8 bAlwaysAllowTranslucentSelect : 1;
 
 	uint8 SelectionOutlineColorIndex;
+
+	FColor OverlayColor;
 #endif
 
 public:
@@ -2552,9 +2558,10 @@ public:
 
 	/** Sets whether this component is being moved by the editor so the renderer can render velocities for it, even when Static. */
 	ENGINE_API void SetIsBeingMovedByEditor(bool bNewIsBeingMoved);
-
-
+	
 	ENGINE_API void SetSelectionOutlineColorIndex(uint8 SelectionOutlineColorIndex);
+	
+	ENGINE_API void SetOverlayColor(FColor OverlayColor);
 #endif// WITH_EDITOR
 
 	/** Resets the cached scene velocity. Useful to prevent motion blur when teleporting components. See also SetIsBeingMovedByEditor(). */
diff --git a/Engine/Source/Runtime/Engine/Private/Components/PrimitiveComponent.cpp b/Engine/Source/Runtime/Engine/Private/Components/PrimitiveComponent.cpp
--- a/Engine/Source/Runtime/Engine/Private/Components/PrimitiveComponent.cpp
+++ b/Engine/Source/Runtime/Engine/Private/Components/PrimitiveComponent.cpp
@@ -319,6 +319,7 @@ UPrimitiveComponent::UPrimitiveComponent(const FObjectInitializer& ObjectInitial
 	IndirectLightingCacheQuality = ILCQ_Point;
 	bStaticWhenNotMoveable = true;
 	bSelectable = true;
+	bWantsEditorEffects = false;
 #if WITH_EDITORONLY_DATA
 	bConsiderForActorPlacementWhenHidden = false;
 #endif // WITH_EDITORONLY_DATA
@@ -390,6 +391,8 @@ UPrimitiveComponent::UPrimitiveComponent(const FObjectInitializer& ObjectInitial
 #if WITH_EDITOR
 	bAlwaysAllowTranslucentSelect = false;
 
+	OverlayColor = FColor(ForceInitToZero);
+	
 	SelectionOutlineColorIndex = 0;
 #endif
 
@@ -1933,11 +1936,26 @@ void UPrimitiveComponent::SetIsBeingMovedByEditor(bool bIsBeingMoved)
 
 void UPrimitiveComponent::SetSelectionOutlineColorIndex(uint8 InSelectionOutlineColorIndex)
 {
+	bool bShouldOverride = InSelectionOutlineColorIndex != 0;
 	SelectionOutlineColorIndex = InSelectionOutlineColorIndex;
+	bWantsEditorEffects |= bShouldOverride;
 	
 	if (SceneProxy)
 	{
 		SceneProxy->SetSelectionOutlineColorIndex_GameThread(InSelectionOutlineColorIndex);
+		SceneProxy->SetSelectionOverride_GameThread(bShouldOverride);
+	}
+}
+
+void UPrimitiveComponent::SetOverlayColor(FColor InOverlayColor)
+{
+	OverlayColor = InOverlayColor;
+	bWantsEditorEffects = true;
+
+	if (SceneProxy)
+	{
+		SceneProxy->SetOverlayColor_GameThread(InOverlayColor);
+		SceneProxy->SetSelectionOverride_GameThread(true);
 	}
 }
 
diff --git a/Engine/Source/Runtime/Engine/Private/HitProxies.cpp b/Engine/Source/Runtime/Engine/Private/HitProxies.cpp
--- a/Engine/Source/Runtime/Engine/Private/HitProxies.cpp
+++ b/Engine/Source/Runtime/Engine/Private/HitProxies.cpp
@@ -74,7 +74,7 @@ static struct FForceInitHitProxyBeforeMain
 
 FHitProxyId::FHitProxyId(FColor Color)
 {
-	Index = ((int32)Color.R << 16) | ((int32)Color.G << 8) | ((int32)Color.B << 0);
+	Index = ((int32)Color.R << 16) | ((int32)Color.G << 8) | ((int32)Color.B << 0) | ((int32)Color.A << 24);
 }
 
 FColor FHitProxyId::GetColor() const
@@ -83,7 +83,7 @@ FColor FHitProxyId::GetColor() const
 		((Index >> 16) & 0xff),
 		((Index >> 8) & 0xff),
 		((Index >> 0) & 0xff),
-		0
+		((Index >> 24) & 0xff)
 		);
 }
 
diff --git a/Engine/Source/Runtime/Engine/Private/PrimitiveSceneProxy.cpp b/Engine/Source/Runtime/Engine/Private/PrimitiveSceneProxy.cpp
--- a/Engine/Source/Runtime/Engine/Private/PrimitiveSceneProxy.cpp
+++ b/Engine/Source/Runtime/Engine/Private/PrimitiveSceneProxy.cpp
@@ -303,13 +303,14 @@ void FPrimitiveSceneProxyDesc::InitializeFrom(const UPrimitiveComponent* InCompo
 	bHiddenInSceneCapture = InComponent->bHiddenInSceneCapture;
 	bRayTracingFarField = InComponent->bRayTracingFarField;
 	bHoldout = InComponent->bHoldout;
-
+	bWantsEditorEffects = InComponent->bWantsEditorEffects;
+	
 	bIsVisible = InComponent->IsVisible();
 	bIsVisibleEditor = InComponent->GetVisibleFlag();
 	bSelected = InComponent->IsSelected();
 	bIndividuallySelected = InComponent->IsComponentIndividuallySelected();
 	bShouldRenderSelected = InComponent->ShouldRenderSelected();
-	bCollisionEnabled = InComponent->IsCollisionEnabled(); 
+	bCollisionEnabled = InComponent->IsCollisionEnabled();
 	
 	if (const AActor* ActorOwner = InComponent->GetOwner())
 	{
@@ -373,6 +374,7 @@ void FPrimitiveSceneProxyDesc::InitializeFrom(const UPrimitiveComponent* InCompo
 
 #if WITH_EDITOR
 	HiddenEditorViews = InComponent->GetHiddenEditorViews();
+	OverlayColor = InComponent->OverlayColor;
 #endif
 	bShouldRenderProxyFallbackToDefaultMaterial = InComponent->ShouldRenderProxyFallbackToDefaultMaterial();
 
@@ -478,6 +480,7 @@ FPrimitiveSceneProxy::FPrimitiveSceneProxy(const FPrimitiveSceneProxyDesc& InPro
 ,	bShouldNotifyOnWorldAddRemove(false)
 ,	bWantsSelectionOutline(true)
 ,	bVerifyUsedMaterials(true)
+,   bWantsEditorEffects(InProxyDesc.bWantsEditorEffects)
 ,	bAllowApproximateOcclusion(InProxyDesc.Mobility != EComponentMobility::Movable)
 ,   bHoldout(InProxyDesc.bHoldout)
 ,	bSplineMesh(false)
@@ -513,6 +516,7 @@ FPrimitiveSceneProxy::FPrimitiveSceneProxy(const FPrimitiveSceneProxyDesc& InPro
 #if WITH_EDITOR
 // by default we are always drawn
 ,	HiddenEditorViews(0)
+,   OverlayColor(InProxyDesc.OverlayColor)
 ,   SelectionOutlineColorIndex(0)
 ,	DrawInAnyEditMode(0)
 ,   bIsFoliage(false)
@@ -1337,6 +1341,17 @@ void FPrimitiveSceneProxy::SetIsBeingMovedByEditor_GameThread(bool bIsBeingMoved
 		});
 }
 
+void FPrimitiveSceneProxy::SetSelectionOverride_GameThread(bool bForceSelection)
+{
+	check(IsInGameThread());
+
+	ENQUEUE_RENDER_COMMAND(SetSelectionOutlineColorIndex)(
+		[this, bForceSelection](FRHICommandListImmediate&)
+		{
+			bWantsEditorEffects = bForceSelection;
+		});
+}
+
 void FPrimitiveSceneProxy::SetSelectionOutlineColorIndex_GameThread(uint8 ColorIndex)
 {
 	check(IsInGameThread());
@@ -1350,6 +1365,17 @@ void FPrimitiveSceneProxy::SetSelectionOutlineColorIndex_GameThread(uint8 ColorI
 			SelectionOutlineColorIndex = ColorIndex;
 		});
 }
+
+void FPrimitiveSceneProxy::SetOverlayColor_GameThread(FColor InOverlayColor)
+{
+	check(IsInGameThread());
+
+	ENQUEUE_RENDER_COMMAND(SetSelectionOutlineColorIndex)(
+		[this, InOverlayColor](FRHICommandListImmediate&)
+		{
+			OverlayColor = InOverlayColor;
+		});
+}
 #endif
 
 #if !(UE_BUILD_SHIPPING || UE_BUILD_TEST)
diff --git a/Engine/Source/Runtime/Engine/Private/StaticMeshRender.cpp b/Engine/Source/Runtime/Engine/Private/StaticMeshRender.cpp
--- a/Engine/Source/Runtime/Engine/Private/StaticMeshRender.cpp
+++ b/Engine/Source/Runtime/Engine/Private/StaticMeshRender.cpp
@@ -2110,7 +2110,7 @@ FPrimitiveViewRelevance FStaticMeshSceneProxy::GetViewRelevance(const FSceneView
 #if WITH_EDITOR
 		//only check these in the editor
 		Result.bEditorVisualizeLevelInstanceRelevance = IsEditingLevelInstanceChild();
-		Result.bEditorStaticSelectionRelevance = (IsSelected() || IsHovered());
+		Result.bEditorStaticSelectionRelevance = (WantsEditorEffects() || IsSelected() || IsHovered());
 #endif
 	}
 
diff --git a/Engine/Source/Runtime/Engine/Public/PrimitiveSceneProxy.h b/Engine/Source/Runtime/Engine/Public/PrimitiveSceneProxy.h
--- a/Engine/Source/Runtime/Engine/Public/PrimitiveSceneProxy.h
+++ b/Engine/Source/Runtime/Engine/Public/PrimitiveSceneProxy.h
@@ -255,18 +255,28 @@ public:
 	 */
 	void SetIsBeingMovedByEditor_GameThread(bool bIsBeingMoved);
 
+	/**
+	 * Force the proxy to render as though it is selected. Shows outlines and overlays.
+	 */
+	void SetSelectionOverride_GameThread(bool bForceSelection);
+	
 	/**
 	 * Enqueue updated selection outline color for the render thread to use.
 	 */
 	void SetSelectionOutlineColorIndex_GameThread(uint8 ColorIndex);
-#endif
+
+	/**
+	 * Enqueue updated overlay color for the render thread to use.
+	 */
+	void SetOverlayColor_GameThread(FColor OverlayColor);
+#endif // WITH_EDITOR
 
 #if !(UE_BUILD_SHIPPING || UE_BUILD_TEST)
 	/**
 	 * Enqueue and update for the render thread to notify it that the primitive color changed.
 	 */
 	void SetPrimitiveColor_GameThread(const FLinearColor& InPrimitiveColor);
-#endif	// WITH_EDITOR
+#endif	
 
 	/** Enqueue and update for the render thread to remove the velocity data for this component from the scene. */
 	ENGINE_API void ResetSceneVelocity_GameThread();
@@ -680,6 +690,7 @@ public:
 	inline bool IsIndividuallySelected() const { return bIndividuallySelected; }
 	inline bool IsEditingLevelInstanceChild() const { return bLevelInstanceEditingState; }
 	inline bool IsSelected() const { return IsParentSelected() || IsIndividuallySelected(); }
+	inline bool WantsEditorEffects() const { return bWantsEditorEffects; }
 	inline bool WantsSelectionOutline() const { return bWantsSelectionOutline; }
 	ENGINE_API bool ShouldRenderCustomDepth() const;
 	inline bool IsVisibleInSceneCaptureOnly() const { return bVisibleInSceneCaptureOnly; }
@@ -753,6 +764,7 @@ public:
 
 #if WITH_EDITOR
 	inline uint8 GetSelectionOutlineColorIndex() const { return SelectionOutlineColorIndex; }
+	inline FColor GetOverlayColor() const { return OverlayColor; }
 #endif // WITH_EDITOR
 	
 	inline bool UseEditorCompositing(const FSceneView* View) const { return GIsEditor && bUseEditorCompositing && !View->bIsGameView; }
@@ -1372,6 +1384,9 @@ protected:
 
 	uint8 bVerifyUsedMaterials : 1;
 
+	/** False by default, if true the proxy wants editor-only effects like outlines and overlays. */
+    uint8 bWantsEditorEffects : 1;
+    
 	/** If this is True, this primitive doesn't need exact occlusion info. */
 	uint8 bAllowApproximateOcclusion : 1;
 
@@ -1509,6 +1524,10 @@ private:
 	/** A copy of the actor's group membership for handling per-view group hiding */
 	uint64 HiddenEditorViews;
 
+	/** Color to blend over the object as an overlay in the viewport */
+	FColor OverlayColor;
+
+	/** Index of the color to use for the object's outline */
 	uint32 SelectionOutlineColorIndex : 8;
 
 	/** Whether this should only draw in any editing mode*/
diff --git a/Engine/Source/Runtime/Engine/Public/PrimitiveSceneProxyDesc.h b/Engine/Source/Runtime/Engine/Public/PrimitiveSceneProxyDesc.h
--- a/Engine/Source/Runtime/Engine/Public/PrimitiveSceneProxyDesc.h
+++ b/Engine/Source/Runtime/Engine/Public/PrimitiveSceneProxyDesc.h
@@ -126,6 +126,7 @@ struct FPrimitiveSceneProxyDesc
 	uint32 bSelected : 1;
 	uint32 bIndividuallySelected : 1;
 	uint32 bShouldRenderSelected : 1;
+	uint32 bWantsEditorEffects : 1;
 	uint32 bCollisionEnabled : 1;
 	uint32 bIsHidden : 1;
 	uint32 bIsHiddenEd : 1;
@@ -181,6 +182,9 @@ struct FPrimitiveSceneProxyDesc
 
 	uint64 HiddenEditorViews = 0;	
 
+#if WITH_EDITOR
+	FColor OverlayColor;
+#endif
 #if MESH_DRAW_COMMAND_STATS
 	FName MeshDrawCommandStatsCategory;
 	FName GetMeshDrawCommandStatsCategory() const { return MeshDrawCommandStatsCategory; }
diff --git a/Engine/Source/Runtime/Renderer/Private/PostProcess/PostProcessSelectionOutline.cpp b/Engine/Source/Runtime/Renderer/Private/PostProcess/PostProcessSelectionOutline.cpp
--- a/Engine/Source/Runtime/Renderer/Private/PostProcess/PostProcessSelectionOutline.cpp
+++ b/Engine/Source/Runtime/Renderer/Private/PostProcess/PostProcessSelectionOutline.cpp
@@ -31,6 +31,8 @@ public:
 		SHADER_PARAMETER_SAMPLER(SamplerState, DepthSampler)
 		SHADER_PARAMETER_RDG_TEXTURE(Texture2D, EditorPrimitivesDepth)
 		SHADER_PARAMETER_RDG_TEXTURE_SRV(Texture2D, EditorPrimitivesStencil)
+		SHADER_PARAMETER_RDG_TEXTURE(Texture2D, OverlayLookupTexture)
+		SHADER_PARAMETER_SAMPLER(SamplerState, OverlayLookupSampler)
 		SHADER_PARAMETER(FScreenTransform, ColorToDepth)
 		SHADER_PARAMETER_ARRAY(FVector4f, OutlineColors, [8])
 		SHADER_PARAMETER(int, OutlineColorIndexBits)
@@ -94,6 +96,7 @@ FScreenPassTexture AddSelectionOutlinePass(
 	// Patch uniform buffers with updated state for rendering the outline mesh draw commands.
 	const bool bIsInstancedStereoView = View.bIsInstancedStereoEnabled && IStereoRendering::IsStereoEyePass(View.StereoPass);
 	const FViewInfo* EditorView = CreateCompositePrimitiveView(View, bIsInstancedStereoView ? View.ViewRect : Inputs.SceneColor.ViewRect, NumSamples);
+	FRDGTextureRef OverlayColorTexture = nullptr;
 
 	// Generate custom depth / stencil for outline shapes.
 	{
@@ -102,6 +105,17 @@ FScreenPassTexture AddSelectionOutlinePass(
 
 		FScene* Scene = View.Family->Scene->GetRenderScene();
 
+		{
+			FRDGTextureDesc OverlayColorDesc = Inputs.SceneColor.Texture->Desc;
+			OverlayColorDesc.Reset();
+			OverlayColorDesc.Format = PF_R8G8B8A8;
+			OverlayColorDesc.ClearValue = FClearValueBinding::Transparent;
+			OverlayColorDesc.Flags = TexCreate_RenderTargetable | TexCreate_ShaderResource;
+			OverlayColorDesc.NumSamples = NumSamples;
+
+			OverlayColorTexture = GraphBuilder.CreateTexture(OverlayColorDesc, TEXT("Editor.ColorOverlay"));
+		}
+		
 		if (View.ShouldRenderView() || !DepthStencilTexture)
 		{
 
@@ -130,6 +144,7 @@ FScreenPassTexture AddSelectionOutlinePass(
 
 				PassParameters->View = EditorView->GetShaderParameters();
 				PassParameters->SceneTextures = Inputs.SceneTextures;
+				PassParameters->RenderTargets[0] = FRenderTargetBinding(OverlayColorTexture, ERenderTargetLoadAction::EClear);
 				PassParameters->RenderTargets.DepthStencil = FDepthStencilBinding(
 					DepthStencilTexture,
 					ERenderTargetLoadAction::EClear,
@@ -259,6 +274,8 @@ FScreenPassTexture AddSelectionOutlinePass(
 		PassParameters->DepthSampler = PointClampSampler;
 		PassParameters->EditorPrimitivesDepth = DepthStencilTexture;
 		PassParameters->EditorPrimitivesStencil = GraphBuilder.CreateSRV(FRDGTextureSRVDesc::CreateWithPixelFormat(DepthStencilTexture, PF_X24_G8));
+		PassParameters->OverlayLookupTexture = OverlayColorTexture;
+		PassParameters->OverlayLookupSampler = PointClampSampler;
 		PassParameters->OutlineColors[0] = View.SelectionOutlineColor;
 		PassParameters->OutlineColors[1] = View.SubduedSelectionOutlineColor;
 		for (int OutlineColorIndex = 2; OutlineColorIndex < PassParameters->OutlineColors.Num(); ++OutlineColorIndex)
diff --git a/Engine/Source/Runtime/Renderer/Private/RendererScene.cpp b/Engine/Source/Runtime/Renderer/Private/RendererScene.cpp
--- a/Engine/Source/Runtime/Renderer/Private/RendererScene.cpp
+++ b/Engine/Source/Runtime/Renderer/Private/RendererScene.cpp
@@ -5957,7 +5957,7 @@ void FScene::Update(FRDGBuilder& GraphBuilder, const FUpdateParameters& Paramete
 					PrimitiveVirtualTextureLod.AddUninitialized();
 					PrimitiveOcclusionBounds.AddUninitialized();
 #if WITH_EDITOR
-					PrimitivesSelected.Add(PrimitiveSceneInfo->Proxy->IsSelected());
+					PrimitivesSelected.Add(PrimitiveSceneInfo->Proxy->WantsEditorEffects() || PrimitiveSceneInfo->Proxy->IsSelected());
 #endif
 #if RHI_RAYTRACING
 					PrimitiveRayTracingFlags.AddZeroed();
diff --git a/Engine/Source/Runtime/Renderer/Private/SceneHitProxyRendering.cpp b/Engine/Source/Runtime/Renderer/Private/SceneHitProxyRendering.cpp
--- a/Engine/Source/Runtime/Renderer/Private/SceneHitProxyRendering.cpp
+++ b/Engine/Source/Runtime/Renderer/Private/SceneHitProxyRendering.cpp
@@ -937,11 +937,15 @@ bool FEditorSelectionMeshProcessor::TryAddMeshBatch(const FMeshBatch& RESTRICT M
 
 void FEditorSelectionMeshProcessor::AddMeshBatch(const FMeshBatch& RESTRICT MeshBatch, uint64 BatchElementMask, const FPrimitiveSceneProxy* RESTRICT PrimitiveSceneProxy, int32 StaticMeshId)
 {
+	if (!PrimitiveSceneProxy)
+	{
+		return;
+	}
+	const bool bWantsEditorEffects = PrimitiveSceneProxy->WantsEditorEffects();
+	const bool bWantsOutlineForSelection = PrimitiveSceneProxy->WantsSelectionOutline() && (PrimitiveSceneProxy->IsSelected() || PrimitiveSceneProxy->IsHovered());
 	if (MeshBatch.bUseForMaterial 
 		&& MeshBatch.bUseSelectionOutline 
-		&& PrimitiveSceneProxy
-		&& PrimitiveSceneProxy->WantsSelectionOutline() 
-		&& (PrimitiveSceneProxy->IsSelected() || PrimitiveSceneProxy->IsHovered()))
+		&& (bWantsEditorEffects || bWantsOutlineForSelection))
 	{
 		const FMaterialRenderProxy* MaterialRenderProxy = MeshBatch.MaterialRenderProxy;
 		while (MaterialRenderProxy)
@@ -989,8 +993,8 @@ bool FEditorSelectionMeshProcessor::Process(
 	const int32 StencilRef = GetStencilValue(ViewIfDynamicMeshCommand, PrimitiveSceneProxy);
 	PassDrawRenderState.SetStencilRef(StencilRef);
 
-	FHitProxyId DummyId;
-	FHitProxyShaderElementData ShaderElementData(DummyId);
+	const FHitProxyId OverlayColor = PrimitiveSceneProxy->GetOverlayColor();
+	FHitProxyShaderElementData ShaderElementData(OverlayColor);
 	ShaderElementData.InitializeMeshMaterialData(ViewIfDynamicMeshCommand, PrimitiveSceneProxy, MeshBatch, StaticMeshId, false);
 
 	const FMeshDrawCommandSortKey SortKey = CalculateMeshStaticSortKey(HitProxyPassShaders.VertexShader, HitProxyPassShaders.PixelShader);
@@ -1069,7 +1073,7 @@ int32 FEditorSelectionMeshProcessor::GetStencilValue(const FSceneView* View, con
 		StencilValue = EncodeSelectionStencilValue(Color, UniqueId);
 		ProxyToStencilIndex.Add(PrimitiveSceneProxy, StencilValue);
 	}
-	else
+	else if (PrimitiveSceneProxy->IsParentSelected())
 	{
 		int Color = PrimitiveSceneProxy->GetSelectionOutlineColorIndex();
 		if (bActorSelectionColorIsSubdued && (Color == 0))
@@ -1081,6 +1085,10 @@ int32 FEditorSelectionMeshProcessor::GetStencilValue(const FSceneView* View, con
 		StencilValue = EncodeSelectionStencilValue(Color, UniqueId);
 		ActorNameToStencilIndex.Add(PrimitiveSceneProxy->GetOwnerName(), StencilValue);
 	}
+	else
+	{
+		StencilValue = ESelectionStencilValues::NotSelected;
+	}
 
 	return StencilValue;
 }
@@ -1093,7 +1101,7 @@ FEditorSelectionMeshProcessor::FEditorSelectionMeshProcessor(const FScene* Scene
 	ActorNameToStencilIndex.Add(NAME_BSP, 1);
 
 	PassDrawRenderState.SetDepthStencilState(TStaticDepthStencilState<true, CF_DepthNearOrEqual, true, CF_Always, SO_Keep, SO_Keep, SO_Replace>::GetRHI());
-	PassDrawRenderState.SetBlendState(TStaticBlendStateWriteMask<CW_NONE, CW_NONE, CW_NONE, CW_NONE>::GetRHI());
+	PassDrawRenderState.SetBlendState(TStaticBlendState<>::GetRHI());
 }
 
 FMeshPassProcessor* CreateEditorSelectionPassProcessor(ERHIFeatureLevel::Type FeatureLevel, const FScene* Scene, const FSceneView* InViewIfDynamicMeshCommand, FMeshPassDrawListContext* InDrawListContext)
diff --git a/Engine/Source/Runtime/TypedElementFramework/Public/Elements/Columns/TypedElementViewportColumns.h b/Engine/Source/Runtime/TypedElementFramework/Public/Elements/Columns/TypedElementViewportColumns.h
--- a/Engine/Source/Runtime/TypedElementFramework/Public/Elements/Columns/TypedElementViewportColumns.h
+++ b/Engine/Source/Runtime/TypedElementFramework/Public/Elements/Columns/TypedElementViewportColumns.h
@@ -8,13 +8,25 @@
 #include "TypedElementViewportColumns.generated.h"
 
 /**
- * Column to hold the color that the object is outlined with when selected in the viewport
+ * Column to hold the color that the object is outlined with in the viewport
  */
-USTRUCT(meta = (DisplayName = "Viewport Color"))
-struct FTypedElementViewportColorColumn final : public FTypedElementDataStorageColumn
+USTRUCT(meta = (DisplayName = "Viewport Outline Color"))
+struct FTypedElementViewportOutlineColorColumn final : public FTypedElementDataStorageColumn
 {
 	GENERATED_BODY()
 
 	UPROPERTY(meta = (ClampMin = "0", ClampMax = "7"))
 	uint8 SelectionOutlineColorIndex = 0;
 };
+
+/**
+ * Column to hold the color that the object is overlaid with in the viewport
+ */
+USTRUCT(meta = (DisplayName = "Viewport Overlay Color"))
+struct FTypedElementViewportOverlayColorColumn final : public FTypedElementDataStorageColumn
+{
+	GENERATED_BODY()
+
+	UPROPERTY()
+	FColor OverlayColor{};
+};
