[Backout] - CL31715868
Original CL Desc
-----------------------------------------------------------------
Add editor-only color overlays for objects in viewport

#p4v-preflight-copy 31541584

diff --git a/Engine/Plugins/Experimental/TypedElementsDataStorage/Source/TypedElementsDataStorage/Private/Compatibility/TypedElementActorViewportProcessors.cpp b/Engine/Plugins/Experimental/TypedElementsDataStorage/Source/TypedElementsDataStorage/Private/Compatibility/TypedElementActorViewportProcessors.cpp
--- a/Engine/Plugins/Experimental/TypedElementsDataStorage/Source/TypedElementsDataStorage/Private/Compatibility/TypedElementActorViewportProcessors.cpp
+++ b/Engine/Plugins/Experimental/TypedElementsDataStorage/Source/TypedElementsDataStorage/Private/Compatibility/TypedElementActorViewportProcessors.cpp
@@ -6,139 +6,24 @@
 #include "Elements/Columns/TypedElementViewportColumns.h"
 #include "Elements/Columns/TypedElementMiscColumns.h"
 #include "Elements/Framework/TypedElementQueryBuilder.h"
-#include "Elements/Framework/TypedElementRegistry.h"
-#include "HAL/IConsoleManager.h"
 #include "MassActorSubsystem.h"
-#include "Elements/Columns/TypedElementSelectionColumns.h"
-#include "ProfilingDebugging/CpuProfilerTrace.h"
-
-FAutoConsoleCommandWithArgsAndOutputDevice SetOutlineColorConsoleCommand(
-	TEXT("TEDS.Debug.SetOutlineColor"),
-	TEXT("Adds an outline color to selected objects."),
-	FConsoleCommandWithArgsAndOutputDeviceDelegate::CreateLambda([](const TArray<FString>& Args, FOutputDevice& Output)
-		{
-			using namespace TypedElementQueryBuilder;
-			using DSI = ITypedElementDataStorageInterface;
-
-			TRACE_CPUPROFILER_EVENT_SCOPE(TEDS.Debug.AddOverlayColorToSelectionCommand);
-
-			if (ITypedElementDataStorageInterface* DataStorage = UTypedElementRegistry::GetInstance()->GetMutableDataStorage())
-			{
-				static TypedElementQueryHandle OverlayQuery = TypedElementInvalidQueryHandle;
-				if (OverlayQuery == TypedElementInvalidQueryHandle)
-				{
-					OverlayQuery = DataStorage->RegisterQuery(
-						Select()
-						.Where()
-							.All<FTypedElementSelectionColumn>()
-						.Compile());
-				}
-				
-				if (OverlayQuery == TypedElementInvalidQueryHandle)
-				{
-					return;
-				}
-
-				if (Args.IsEmpty())
-				{
-					Output.Log(TEXT("Provide a color index (0-7) to use as outline"));
-					return;
-				}
-				
-				// Parse the color
-				int32 ColorIndex;
-				LexFromString(ColorIndex, *Args[0]);
-
-				if (!(ColorIndex >= 0 && ColorIndex <= 7))
-				{
-					Output.Log(TEXT("Color index must be in range [0,7]"));
-					return;
-				}
-
-				TArray<TypedElementRowHandle> RowHandles;
-				
-				DataStorage->RunQuery(OverlayQuery, [ColorIndex, &RowHandles](const DSI::FQueryDescription&, DSI::IDirectQueryContext& Context)
-				{
-					RowHandles = Context.GetRowHandles();
-				});
-
-				for (TypedElementRowHandle Row : RowHandles)
-				{
-					DataStorage->AddOrGetColumn<FTypedElementViewportOutlineColorColumn>(Row)->SelectionOutlineColorIndex = ColorIndex;
-					DataStorage->AddColumn<FTypedElementSyncBackToWorldTag>(Row);
-				}
-			}
-		}));
-
-FAutoConsoleCommandWithArgsAndOutputDevice SetSelectionOverlayColorConsoleCommand(
-	TEXT("TEDS.Debug.SetOverlayColor"),
-	TEXT("Adds an overlay color to selected objects."),
-	FConsoleCommandWithArgsAndOutputDeviceDelegate::CreateLambda([](const TArray<FString>& Args, FOutputDevice& Output)
-		{
-			using namespace TypedElementQueryBuilder;
-			using DSI = ITypedElementDataStorageInterface;
-
-			TRACE_CPUPROFILER_EVENT_SCOPE(TEDS.Debug.AddOverlayColorToSelectionCommand);
-
-			if (ITypedElementDataStorageInterface* DataStorage = UTypedElementRegistry::GetInstance()->GetMutableDataStorage())
-			{
-				static TypedElementQueryHandle OverlayQuery = TypedElementInvalidQueryHandle;
-				if (OverlayQuery == TypedElementInvalidQueryHandle)
-				{
-					OverlayQuery = DataStorage->RegisterQuery(
-						Select()
-						.Where()
-							.All<FTypedElementSelectionColumn>()
-						.Compile());
-				}
-				
-				if (OverlayQuery == TypedElementInvalidQueryHandle)
-				{
-					return;
-				}
-
-				if (Args.IsEmpty())
-				{
-					Output.Log(TEXT("Provide a color in hexadecimal format (#RRGGBBAA) to overlay."));
-					return;
-				}
-				
-				// Parse the color
-				FColor Color = FColor::FromHex(Args[0]);
-				Color.A = FMath::Clamp(Color.A, 0, 128);
-
-				TArray<TypedElementRowHandle> RowHandles;
-				
-				DataStorage->RunQuery(OverlayQuery, [Color, &RowHandles](const DSI::FQueryDescription&, DSI::IDirectQueryContext& Context)
-				{
-					RowHandles = Context.GetRowHandles();
-				});
-
-				for (TypedElementRowHandle Row : RowHandles)
-				{
-					DataStorage->AddOrGetColumn<FTypedElementViewportOverlayColorColumn>(Row)->OverlayColor = Color;
-					DataStorage->AddColumn<FTypedElementSyncBackToWorldTag>(Row);
-				}
-			}
-		}));
 
 void UTypedElementActorViewportFactory::RegisterQueries(ITypedElementDataStorageInterface& DataStorage)
 {
-	RegisterOutlineColorColumnToActor(DataStorage);
-	RegisterOverlayColorColumnToActor(DataStorage);
+	RegisterSelectionOutlineColorColumnToActor(DataStorage);
 }
 
-void UTypedElementActorViewportFactory::RegisterOutlineColorColumnToActor(ITypedElementDataStorageInterface& DataStorage)
+void UTypedElementActorViewportFactory::RegisterSelectionOutlineColorColumnToActor(ITypedElementDataStorageInterface& DataStorage) const
 {
 	using namespace TypedElementQueryBuilder;
 	using DSI = ITypedElementDataStorageInterface;
 
 	DataStorage.RegisterQuery(
 		Select(
-			TEXT("Sync viewport outline color column to actor"),
+			TEXT("Sync selection outline color column to actor"),
 			FProcessor(DSI::EQueryTickPhase::DuringPhysics, DataStorage.GetQueryTickGroupName(DSI::EQueryTickGroups::SyncDataStorageToExternal))
 				.ForceToGameThread(true),
-			[](FMassActorFragment& Actor, const FTypedElementViewportOutlineColorColumn& ViewportColor)
+			[](FMassActorFragment& Actor, const FTypedElementViewportColorColumn& ViewportColor)
 			{
 				if (AActor* ActorInstance = Actor.GetMutable(); ActorInstance != nullptr)
 				{
@@ -154,30 +39,3 @@ void UTypedElementActorViewportFactory::RegisterOutlineColorColumnToActor(ITyped
 		.Compile()
 	);
 }
-
-void UTypedElementActorViewportFactory::RegisterOverlayColorColumnToActor(ITypedElementDataStorageInterface& DataStorage)
-{
-	using namespace TypedElementQueryBuilder;
-	using DSI = ITypedElementDataStorageInterface;
-
-	DataStorage.RegisterQuery(
-		Select(
-			TEXT("Sync viewport overlay color column to actor"),
-			FProcessor(DSI::EQueryTickPhase::DuringPhysics, DataStorage.GetQueryTickGroupName(DSI::EQueryTickGroups::SyncDataStorageToExternal))
-				.ForceToGameThread(true),
-			[](FMassActorFragment& Actor, const FTypedElementViewportOverlayColorColumn& ViewportColor)
-			{
-				if (AActor* ActorInstance = Actor.GetMutable(); ActorInstance != nullptr)
-				{
-					if (UPrimitiveComponent* PrimitiveComponent = ActorInstance->GetComponentByClass<UPrimitiveComponent>())
-					{
-						PrimitiveComponent->SetOverlayColor(ViewportColor.OverlayColor);
-					}
-				}
-			}
-		)
-		.Where()
-			.All<FTypedElementSyncBackToWorldTag>()
-		.Compile()
-	);
-}
diff --git a/Engine/Plugins/Experimental/TypedElementsDataStorage/Source/TypedElementsDataStorage/Private/Compatibility/TypedElementActorViewportProcessors.h b/Engine/Plugins/Experimental/TypedElementsDataStorage/Source/TypedElementsDataStorage/Private/Compatibility/TypedElementActorViewportProcessors.h
--- a/Engine/Plugins/Experimental/TypedElementsDataStorage/Source/TypedElementsDataStorage/Private/Compatibility/TypedElementActorViewportProcessors.h
+++ b/Engine/Plugins/Experimental/TypedElementsDataStorage/Source/TypedElementsDataStorage/Private/Compatibility/TypedElementActorViewportProcessors.h
@@ -18,6 +18,5 @@ public:
 	void RegisterQueries(ITypedElementDataStorageInterface& DataStorage) override;
 
 private:
-	void RegisterOutlineColorColumnToActor(ITypedElementDataStorageInterface& DataStorage);
-	void RegisterOverlayColorColumnToActor(ITypedElementDataStorageInterface& DataStorage);
+	void RegisterSelectionOutlineColorColumnToActor(ITypedElementDataStorageInterface& DataStorage) const;
 };
diff --git a/Engine/Plugins/Experimental/TypedElementsDataStorage/Source/TypedElementsDataStorageRevisionControl/Source/Private/RevisionControlProcessors.cpp b/Engine/Plugins/Experimental/TypedElementsDataStorage/Source/TypedElementsDataStorageRevisionControl/Source/Private/RevisionControlProcessors.cpp
--- a/Engine/Plugins/Experimental/TypedElementsDataStorage/Source/TypedElementsDataStorageRevisionControl/Source/Private/RevisionControlProcessors.cpp
+++ b/Engine/Plugins/Experimental/TypedElementsDataStorage/Source/TypedElementsDataStorageRevisionControl/Source/Private/RevisionControlProcessors.cpp
@@ -99,7 +99,7 @@ void UTypedElementRevisionControlFactory::RegisterQueries(ITypedElementDataStora
 				Context.RunSubquery(0, PackageReference.Row, CreateSubqueryCallbackBinding(
 					[&Context, &ObjectRow](DSI::ISubqueryContext& SubQueryContext)
 					{
-						Context.AddColumn<FTypedElementViewportOutlineColorColumn>(ObjectRow, { .SelectionOutlineColorIndex = DetermineOutlineColorIndex(SubQueryContext) });
+						Context.AddColumn<FTypedElementViewportColorColumn>(ObjectRow, { .SelectionOutlineColorIndex = DetermineOutlineColorIndex(SubQueryContext) });
 						Context.AddColumns<FTypedElementSyncBackToWorldTag>(ObjectRow);
 					})
 				);
diff --git a/Engine/Shaders/Private/PostProcessSelectionOutline.usf b/Engine/Shaders/Private/PostProcessSelectionOutline.usf
--- a/Engine/Shaders/Private/PostProcessSelectionOutline.usf
+++ b/Engine/Shaders/Private/PostProcessSelectionOutline.usf
@@ -20,9 +20,6 @@ SamplerState ColorSampler;
 Texture2D DepthTexture;
 SamplerState DepthSampler;
 
-Texture2D OverlayLookupTexture;
-SamplerState OverlayLookupSampler;
-
 SCREEN_PASS_TEXTURE_VIEWPORT(Color)
 SCREEN_PASS_TEXTURE_VIEWPORT(Depth)
 FScreenTransform ColorToDepth;
@@ -236,36 +233,32 @@ void MainPS(
 	);
 #endif
 
-	const float2 DepthUV = ApplyScreenTransform(ColorUV, ColorToDepth);
+	const int ScoutSum = dot(saturate(StencilScout), 1);
 
-	float Center = Texture2DSampleLevel(DepthTexture, DepthSampler, DepthUV, 0).r;
-	float Left   = Texture2DSampleLevel(DepthTexture, DepthSampler, DepthUV + float2(-1,  0) * Depth_ExtentInverse, 0).r;
-	float Right  = Texture2DSampleLevel(DepthTexture, DepthSampler, DepthUV + float2( 1,  0) * Depth_ExtentInverse, 0).r;
-	float Top    = Texture2DSampleLevel(DepthTexture, DepthSampler, DepthUV + float2( 0, -1) * Depth_ExtentInverse, 0).r;
-	float Bottom = Texture2DSampleLevel(DepthTexture, DepthSampler, DepthUV + float2( 0,  1) * Depth_ExtentInverse, 0).r;
+	// If this sum is zero, none of our neighbors are selected pixels and we can skip all the heavy processing.
+	if (ScoutSum > 0)
+	{
+		const float2 DepthUV = ApplyScreenTransform(ColorUV, ColorToDepth);
 
-	// This allows to reconstruct depth with a small pixel offset without many texture lookups (4xMSAA * 5 neighbors -> 20 samples)
-	// It's an approximation assuming the surface is a plane.
-	float3 DeviceZPlane;
-	DeviceZPlane.x = (Right - Left) / 2;
-	DeviceZPlane.y = (Bottom - Top) / 2;
-	DeviceZPlane.z = Center;
+		float Center = Texture2DSampleLevel(DepthTexture, DepthSampler, DepthUV, 0).r;
+		float Left   = Texture2DSampleLevel(DepthTexture, DepthSampler, DepthUV + float2(-1,  0) * Depth_ExtentInverse, 0).r;
+		float Right  = Texture2DSampleLevel(DepthTexture, DepthSampler, DepthUV + float2( 1,  0) * Depth_ExtentInverse, 0).r;
+		float Top    = Texture2DSampleLevel(DepthTexture, DepthSampler, DepthUV + float2( 0, -1) * Depth_ExtentInverse, 0).r;
+		float Bottom = Texture2DSampleLevel(DepthTexture, DepthSampler, DepthUV + float2( 0,  1) * Depth_ExtentInverse, 0).r;
 
-	float2 DeviceZMinMax;
-	DeviceZMinMax.x = min(Center, min(min(Left, Right), min(Top, Bottom)));
-	DeviceZMinMax.y = max(Center, max(max(Left, Right), max(Top, Bottom)));
+		// This allows to reconstruct depth with a small pixel offset without many texture lookups (4xMSAA * 5 neighbors -> 20 samples)
+		// It's an approximation assuming the surface is a plane.
+		float3 DeviceZPlane;
+		DeviceZPlane.x = (Right - Left) / 2;
+		DeviceZPlane.y = (Bottom - Top) / 2;
+		DeviceZPlane.z = Center;
 
-	FPixelInfo CenterPixelInfo = GetPixelInfo(ColorPixelPos, 0, SecondaryViewportOffset, 0, DeviceZPlane, DeviceZMinMax);
+		float2 DeviceZMinMax;
+		DeviceZMinMax.x = min(Center, min(min(Left, Right), min(Top, Bottom)));
+		DeviceZMinMax.y = max(Center, max(max(Left, Right), max(Top, Bottom)));
 
-	float4 OverlayColor = Texture2DSample(OverlayLookupTexture, OverlayLookupSampler, ColorUV);
-	OutColor = Texture2DSample(ColorTexture, ColorSampler, ColorUV); 
-	
-	float4 SelectionEffect;
-	const int ScoutSum = dot(saturate(StencilScout), 1);
+		FPixelInfo CenterPixelInfo = GetPixelInfo(ColorPixelPos, 0, SecondaryViewportOffset, 0, DeviceZPlane, DeviceZMinMax);
 
-	// If this sum is zero, none of our neighbors are selected pixels and we can skip all the heavy processing.
-	if (ScoutSum > 0)
-	{
 		float4 Sum = 0;
 #if COMPILER_GLSL && FEATURE_LEVEL >= FEATURE_LEVEL_SM4
 	UNROLL
@@ -274,47 +267,49 @@ void MainPS(
 		{
 			Sum += PerSample(ColorPixelPos, SampleID, CenterPixelInfo, DeviceZPlane, DeviceZMinMax);
 		}
-		SelectionEffect = Sum / MSAA_SAMPLE_COUNT;
-	}
-	
-	// Scene color has gamma applied. Need to remove and re-apply after linear operations.
-	if (SelectionEffect.a > 0 || OverlayColor.a > 0)
-	{
-#if SELECTION_OUTLINE_HDR
-		static const float ScRGBScaleFactor = 80.0f; // used to convert between ScRGB and nits
-		SelectionEffect.rgb *= abs(UILuminanceAndIsSCRGB);
-		OverlayColor.rgb *= abs(UILuminanceAndIsSCRGB);
-		if (UILuminanceAndIsSCRGB > 0.0)
-		{
-			OutColor.rgb *= ScRGBScaleFactor;
-		}
-		else
+		float4 Avg = Sum / MSAA_SAMPLE_COUNT;
+
+		OutColor = Texture2DSample(ColorTexture, ColorSampler, ColorUV);
+
+		// Scene color has gamma applied. Need to remove and re-apply after linear operation.
+		if (Avg.a > 0)
 		{
-			const float3x3 Rec2020_2_sRGB = mul(XYZ_2_sRGB_MAT, Rec2020_2_XYZ_MAT);
-			OutColor.rgb = ST2084ToLinear(OutColor.rgb);
-			OutColor.rgb = mul(Rec2020_2_sRGB, OutColor.rgb);
-		}
+#if SELECTION_OUTLINE_HDR
+			static const float ScRGBScaleFactor = 80.0f; // used to convert between ScRGB and nits
+			Avg.rgb *= abs(UILuminanceAndIsSCRGB);
+			if (UILuminanceAndIsSCRGB > 0.0)
+			{
+				OutColor.rgb *= ScRGBScaleFactor;
+			}
+			else
+			{
+				const float3x3 Rec2020_2_sRGB = mul(XYZ_2_sRGB_MAT, Rec2020_2_XYZ_MAT);
+				OutColor.rgb = ST2084ToLinear(OutColor.rgb);
+				OutColor.rgb = mul(Rec2020_2_sRGB, OutColor.rgb);
+			}
 #else
-		OutColor.rgb = sRGBToLinear(OutColor.rgb);
+			OutColor.rgb = sRGBToLinear(OutColor.rgb);
 #endif
-		float PatternMask = ((ColorPixelPos.x/2 + ColorPixelPos.y/2) % 2) * 0.6f;
-		float VisibleMask = lerp(PatternMask, 1.0f, CenterPixelInfo.fVisible);
-		OutColor.rgb = lerp(OutColor.rgb, OverlayColor.rgb, OverlayColor.a * VisibleMask);
-		OutColor.rgb = lerp(OutColor.rgb, SelectionEffect.rgb, SelectionEffect.a);
+			OutColor.rgb = lerp(OutColor.rgb, Avg.rgb, Avg.a);
 
 #if SELECTION_OUTLINE_HDR
-		if (UILuminanceAndIsSCRGB > 0.0)
-		{
-			OutColor.rgb /= ScRGBScaleFactor;
-		}
-		else
-		{
-			const float3x3 sRGB_2_Rec2020 = mul(XYZ_2_Rec2020_MAT, sRGB_2_XYZ_MAT);
-			OutColor.xyz = mul(sRGB_2_Rec2020, OutColor.xyz);
-			OutColor.rgb = LinearToST2084(OutColor.rgb);
-		}
+			if (UILuminanceAndIsSCRGB > 0.0)
+			{
+				OutColor.rgb /= ScRGBScaleFactor;
+			}
+			else
+			{
+				const float3x3 sRGB_2_Rec2020 = mul(XYZ_2_Rec2020_MAT, sRGB_2_XYZ_MAT);
+				OutColor.xyz = mul(sRGB_2_Rec2020, OutColor.xyz);
+				OutColor.rgb = LinearToST2084(OutColor.rgb);
+			}
 #else
-		OutColor.rgb = LinearToSrgbBranchless(OutColor.rgb);
+			OutColor.rgb = LinearToSrgbBranchless(OutColor.rgb);
 #endif
+		}
+	}
+	else
+	{
+		OutColor = Texture2DSample(ColorTexture, ColorSampler, ColorUV);
 	}
 }
diff --git a/Engine/Source/Runtime/Engine/Classes/Components/PrimitiveComponent.h b/Engine/Source/Runtime/Engine/Classes/Components/PrimitiveComponent.h
--- a/Engine/Source/Runtime/Engine/Classes/Components/PrimitiveComponent.h
+++ b/Engine/Source/Runtime/Engine/Classes/Components/PrimitiveComponent.h
@@ -443,10 +443,6 @@ public:
 	UPROPERTY()
 	uint8 bSelectable:1;
 
-	/** When true, this component requests editor effects like outlines and overlays. */
-	UPROPERTY()
-	uint8 bWantsEditorEffects:1;
-
 #if WITH_EDITORONLY_DATA
 	/** If true, this component will be considered for placement when dragging and placing items in the editor even if it is not visible, such as in the case of hidden collision meshes */
 	UPROPERTY(EditAnywhere, AdvancedDisplay, Category = Collision)
@@ -677,8 +673,6 @@ public:
 	uint8 bAlwaysAllowTranslucentSelect : 1;
 
 	uint8 SelectionOutlineColorIndex;
-
-	FColor OverlayColor;
 #endif
 
 public:
@@ -2558,10 +2552,9 @@ public:
 
 	/** Sets whether this component is being moved by the editor so the renderer can render velocities for it, even when Static. */
 	ENGINE_API void SetIsBeingMovedByEditor(bool bNewIsBeingMoved);
-	
+
+
 	ENGINE_API void SetSelectionOutlineColorIndex(uint8 SelectionOutlineColorIndex);
-	
-	ENGINE_API void SetOverlayColor(FColor OverlayColor);
 #endif// WITH_EDITOR
 
 	/** Resets the cached scene velocity. Useful to prevent motion blur when teleporting components. See also SetIsBeingMovedByEditor(). */
diff --git a/Engine/Source/Runtime/Engine/Private/Components/PrimitiveComponent.cpp b/Engine/Source/Runtime/Engine/Private/Components/PrimitiveComponent.cpp
--- a/Engine/Source/Runtime/Engine/Private/Components/PrimitiveComponent.cpp
+++ b/Engine/Source/Runtime/Engine/Private/Components/PrimitiveComponent.cpp
@@ -319,7 +319,6 @@ UPrimitiveComponent::UPrimitiveComponent(const FObjectInitializer& ObjectInitial
 	IndirectLightingCacheQuality = ILCQ_Point;
 	bStaticWhenNotMoveable = true;
 	bSelectable = true;
-	bWantsEditorEffects = false;
 #if WITH_EDITORONLY_DATA
 	bConsiderForActorPlacementWhenHidden = false;
 #endif // WITH_EDITORONLY_DATA
@@ -391,8 +390,6 @@ UPrimitiveComponent::UPrimitiveComponent(const FObjectInitializer& ObjectInitial
 #if WITH_EDITOR
 	bAlwaysAllowTranslucentSelect = false;
 
-	OverlayColor = FColor(ForceInitToZero);
-	
 	SelectionOutlineColorIndex = 0;
 #endif
 
@@ -1936,26 +1933,11 @@ void UPrimitiveComponent::SetIsBeingMovedByEditor(bool bIsBeingMoved)
 
 void UPrimitiveComponent::SetSelectionOutlineColorIndex(uint8 InSelectionOutlineColorIndex)
 {
-	bool bShouldOverride = InSelectionOutlineColorIndex != 0;
 	SelectionOutlineColorIndex = InSelectionOutlineColorIndex;
-	bWantsEditorEffects |= bShouldOverride;
 	
 	if (SceneProxy)
 	{
 		SceneProxy->SetSelectionOutlineColorIndex_GameThread(InSelectionOutlineColorIndex);
-		SceneProxy->SetSelectionOverride_GameThread(bShouldOverride);
-	}
-}
-
-void UPrimitiveComponent::SetOverlayColor(FColor InOverlayColor)
-{
-	OverlayColor = InOverlayColor;
-	bWantsEditorEffects = true;
-
-	if (SceneProxy)
-	{
-		SceneProxy->SetOverlayColor_GameThread(InOverlayColor);
-		SceneProxy->SetSelectionOverride_GameThread(true);
 	}
 }
 
diff --git a/Engine/Source/Runtime/Engine/Private/HitProxies.cpp b/Engine/Source/Runtime/Engine/Private/HitProxies.cpp
--- a/Engine/Source/Runtime/Engine/Private/HitProxies.cpp
+++ b/Engine/Source/Runtime/Engine/Private/HitProxies.cpp
@@ -74,7 +74,7 @@ static struct FForceInitHitProxyBeforeMain
 
 FHitProxyId::FHitProxyId(FColor Color)
 {
-	Index = ((int32)Color.R << 16) | ((int32)Color.G << 8) | ((int32)Color.B << 0) | ((int32)Color.A << 24);
+	Index = ((int32)Color.R << 16) | ((int32)Color.G << 8) | ((int32)Color.B << 0);
 }
 
 FColor FHitProxyId::GetColor() const
@@ -83,7 +83,7 @@ FColor FHitProxyId::GetColor() const
 		((Index >> 16) & 0xff),
 		((Index >> 8) & 0xff),
 		((Index >> 0) & 0xff),
-		((Index >> 24) & 0xff)
+		0
 		);
 }
 
diff --git a/Engine/Source/Runtime/Engine/Private/PrimitiveSceneProxy.cpp b/Engine/Source/Runtime/Engine/Private/PrimitiveSceneProxy.cpp
--- a/Engine/Source/Runtime/Engine/Private/PrimitiveSceneProxy.cpp
+++ b/Engine/Source/Runtime/Engine/Private/PrimitiveSceneProxy.cpp
@@ -303,14 +303,13 @@ void FPrimitiveSceneProxyDesc::InitializeFrom(const UPrimitiveComponent* InCompo
 	bHiddenInSceneCapture = InComponent->bHiddenInSceneCapture;
 	bRayTracingFarField = InComponent->bRayTracingFarField;
 	bHoldout = InComponent->bHoldout;
-	bWantsEditorEffects = InComponent->bWantsEditorEffects;
-	
+
 	bIsVisible = InComponent->IsVisible();
 	bIsVisibleEditor = InComponent->GetVisibleFlag();
 	bSelected = InComponent->IsSelected();
 	bIndividuallySelected = InComponent->IsComponentIndividuallySelected();
 	bShouldRenderSelected = InComponent->ShouldRenderSelected();
-	bCollisionEnabled = InComponent->IsCollisionEnabled();
+	bCollisionEnabled = InComponent->IsCollisionEnabled(); 
 	
 	if (const AActor* ActorOwner = InComponent->GetOwner())
 	{
@@ -374,7 +373,6 @@ void FPrimitiveSceneProxyDesc::InitializeFrom(const UPrimitiveComponent* InCompo
 
 #if WITH_EDITOR
 	HiddenEditorViews = InComponent->GetHiddenEditorViews();
-	OverlayColor = InComponent->OverlayColor;
 #endif
 	bShouldRenderProxyFallbackToDefaultMaterial = InComponent->ShouldRenderProxyFallbackToDefaultMaterial();
 
@@ -480,7 +478,6 @@ FPrimitiveSceneProxy::FPrimitiveSceneProxy(const FPrimitiveSceneProxyDesc& InPro
 ,	bShouldNotifyOnWorldAddRemove(false)
 ,	bWantsSelectionOutline(true)
 ,	bVerifyUsedMaterials(true)
-,   bWantsEditorEffects(InProxyDesc.bWantsEditorEffects)
 ,	bAllowApproximateOcclusion(InProxyDesc.Mobility != EComponentMobility::Movable)
 ,   bHoldout(InProxyDesc.bHoldout)
 ,	bSplineMesh(false)
@@ -516,7 +513,6 @@ FPrimitiveSceneProxy::FPrimitiveSceneProxy(const FPrimitiveSceneProxyDesc& InPro
 #if WITH_EDITOR
 // by default we are always drawn
 ,	HiddenEditorViews(0)
-,   OverlayColor(InProxyDesc.OverlayColor)
 ,   SelectionOutlineColorIndex(0)
 ,	DrawInAnyEditMode(0)
 ,   bIsFoliage(false)
@@ -1341,17 +1337,6 @@ void FPrimitiveSceneProxy::SetIsBeingMovedByEditor_GameThread(bool bIsBeingMoved
 		});
 }
 
-void FPrimitiveSceneProxy::SetSelectionOverride_GameThread(bool bForceSelection)
-{
-	check(IsInGameThread());
-
-	ENQUEUE_RENDER_COMMAND(SetSelectionOutlineColorIndex)(
-		[this, bForceSelection](FRHICommandListImmediate&)
-		{
-			bWantsEditorEffects = bForceSelection;
-		});
-}
-
 void FPrimitiveSceneProxy::SetSelectionOutlineColorIndex_GameThread(uint8 ColorIndex)
 {
 	check(IsInGameThread());
@@ -1365,17 +1350,6 @@ void FPrimitiveSceneProxy::SetSelectionOutlineColorIndex_GameThread(uint8 ColorI
 			SelectionOutlineColorIndex = ColorIndex;
 		});
 }
-
-void FPrimitiveSceneProxy::SetOverlayColor_GameThread(FColor InOverlayColor)
-{
-	check(IsInGameThread());
-
-	ENQUEUE_RENDER_COMMAND(SetSelectionOutlineColorIndex)(
-		[this, InOverlayColor](FRHICommandListImmediate&)
-		{
-			OverlayColor = InOverlayColor;
-		});
-}
 #endif
 
 #if !(UE_BUILD_SHIPPING || UE_BUILD_TEST)
diff --git a/Engine/Source/Runtime/Engine/Private/StaticMeshRender.cpp b/Engine/Source/Runtime/Engine/Private/StaticMeshRender.cpp
--- a/Engine/Source/Runtime/Engine/Private/StaticMeshRender.cpp
+++ b/Engine/Source/Runtime/Engine/Private/StaticMeshRender.cpp
@@ -2110,7 +2110,7 @@ FPrimitiveViewRelevance FStaticMeshSceneProxy::GetViewRelevance(const FSceneView
 #if WITH_EDITOR
 		//only check these in the editor
 		Result.bEditorVisualizeLevelInstanceRelevance = IsEditingLevelInstanceChild();
-		Result.bEditorStaticSelectionRelevance = (WantsEditorEffects() || IsSelected() || IsHovered());
+		Result.bEditorStaticSelectionRelevance = (IsSelected() || IsHovered());
 #endif
 	}
 
diff --git a/Engine/Source/Runtime/Engine/Public/PrimitiveSceneProxy.h b/Engine/Source/Runtime/Engine/Public/PrimitiveSceneProxy.h
--- a/Engine/Source/Runtime/Engine/Public/PrimitiveSceneProxy.h
+++ b/Engine/Source/Runtime/Engine/Public/PrimitiveSceneProxy.h
@@ -255,28 +255,18 @@ public:
 	 */
 	void SetIsBeingMovedByEditor_GameThread(bool bIsBeingMoved);
 
-	/**
-	 * Force the proxy to render as though it is selected. Shows outlines and overlays.
-	 */
-	void SetSelectionOverride_GameThread(bool bForceSelection);
-	
 	/**
 	 * Enqueue updated selection outline color for the render thread to use.
 	 */
 	void SetSelectionOutlineColorIndex_GameThread(uint8 ColorIndex);
-
-	/**
-	 * Enqueue updated overlay color for the render thread to use.
-	 */
-	void SetOverlayColor_GameThread(FColor OverlayColor);
-#endif // WITH_EDITOR
+#endif
 
 #if !(UE_BUILD_SHIPPING || UE_BUILD_TEST)
 	/**
 	 * Enqueue and update for the render thread to notify it that the primitive color changed.
 	 */
 	void SetPrimitiveColor_GameThread(const FLinearColor& InPrimitiveColor);
-#endif	
+#endif	// WITH_EDITOR
 
 	/** Enqueue and update for the render thread to remove the velocity data for this component from the scene. */
 	ENGINE_API void ResetSceneVelocity_GameThread();
@@ -690,7 +680,6 @@ public:
 	inline bool IsIndividuallySelected() const { return bIndividuallySelected; }
 	inline bool IsEditingLevelInstanceChild() const { return bLevelInstanceEditingState; }
 	inline bool IsSelected() const { return IsParentSelected() || IsIndividuallySelected(); }
-	inline bool WantsEditorEffects() const { return bWantsEditorEffects; }
 	inline bool WantsSelectionOutline() const { return bWantsSelectionOutline; }
 	ENGINE_API bool ShouldRenderCustomDepth() const;
 	inline bool IsVisibleInSceneCaptureOnly() const { return bVisibleInSceneCaptureOnly; }
@@ -764,7 +753,6 @@ public:
 
 #if WITH_EDITOR
 	inline uint8 GetSelectionOutlineColorIndex() const { return SelectionOutlineColorIndex; }
-	inline FColor GetOverlayColor() const { return OverlayColor; }
 #endif // WITH_EDITOR
 	
 	inline bool UseEditorCompositing(const FSceneView* View) const { return GIsEditor && bUseEditorCompositing && !View->bIsGameView; }
@@ -1384,9 +1372,6 @@ protected:
 
 	uint8 bVerifyUsedMaterials : 1;
 
-	/** False by default, if true the proxy wants editor-only effects like outlines and overlays. */
-    uint8 bWantsEditorEffects : 1;
-    
 	/** If this is True, this primitive doesn't need exact occlusion info. */
 	uint8 bAllowApproximateOcclusion : 1;
 
@@ -1524,10 +1509,6 @@ private:
 	/** A copy of the actor's group membership for handling per-view group hiding */
 	uint64 HiddenEditorViews;
 
-	/** Color to blend over the object as an overlay in the viewport */
-	FColor OverlayColor;
-
-	/** Index of the color to use for the object's outline */
 	uint32 SelectionOutlineColorIndex : 8;
 
 	/** Whether this should only draw in any editing mode*/
diff --git a/Engine/Source/Runtime/Engine/Public/PrimitiveSceneProxyDesc.h b/Engine/Source/Runtime/Engine/Public/PrimitiveSceneProxyDesc.h
--- a/Engine/Source/Runtime/Engine/Public/PrimitiveSceneProxyDesc.h
+++ b/Engine/Source/Runtime/Engine/Public/PrimitiveSceneProxyDesc.h
@@ -126,7 +126,6 @@ struct FPrimitiveSceneProxyDesc
 	uint32 bSelected : 1;
 	uint32 bIndividuallySelected : 1;
 	uint32 bShouldRenderSelected : 1;
-	uint32 bWantsEditorEffects : 1;
 	uint32 bCollisionEnabled : 1;
 	uint32 bIsHidden : 1;
 	uint32 bIsHiddenEd : 1;
@@ -182,9 +181,6 @@ struct FPrimitiveSceneProxyDesc
 
 	uint64 HiddenEditorViews = 0;	
 
-#if WITH_EDITOR
-	FColor OverlayColor;
-#endif
 #if MESH_DRAW_COMMAND_STATS
 	FName MeshDrawCommandStatsCategory;
 	FName GetMeshDrawCommandStatsCategory() const { return MeshDrawCommandStatsCategory; }
diff --git a/Engine/Source/Runtime/Renderer/Private/PostProcess/PostProcessSelectionOutline.cpp b/Engine/Source/Runtime/Renderer/Private/PostProcess/PostProcessSelectionOutline.cpp
--- a/Engine/Source/Runtime/Renderer/Private/PostProcess/PostProcessSelectionOutline.cpp
+++ b/Engine/Source/Runtime/Renderer/Private/PostProcess/PostProcessSelectionOutline.cpp
@@ -31,8 +31,6 @@ public:
 		SHADER_PARAMETER_SAMPLER(SamplerState, DepthSampler)
 		SHADER_PARAMETER_RDG_TEXTURE(Texture2D, EditorPrimitivesDepth)
 		SHADER_PARAMETER_RDG_TEXTURE_SRV(Texture2D, EditorPrimitivesStencil)
-		SHADER_PARAMETER_RDG_TEXTURE(Texture2D, OverlayLookupTexture)
-		SHADER_PARAMETER_SAMPLER(SamplerState, OverlayLookupSampler)
 		SHADER_PARAMETER(FScreenTransform, ColorToDepth)
 		SHADER_PARAMETER_ARRAY(FVector4f, OutlineColors, [8])
 		SHADER_PARAMETER(int, OutlineColorIndexBits)
@@ -96,7 +94,6 @@ FScreenPassTexture AddSelectionOutlinePass(
 	// Patch uniform buffers with updated state for rendering the outline mesh draw commands.
 	const bool bIsInstancedStereoView = View.bIsInstancedStereoEnabled && IStereoRendering::IsStereoEyePass(View.StereoPass);
 	const FViewInfo* EditorView = CreateCompositePrimitiveView(View, bIsInstancedStereoView ? View.ViewRect : Inputs.SceneColor.ViewRect, NumSamples);
-	FRDGTextureRef OverlayColorTexture = nullptr;
 
 	// Generate custom depth / stencil for outline shapes.
 	{
@@ -105,17 +102,6 @@ FScreenPassTexture AddSelectionOutlinePass(
 
 		FScene* Scene = View.Family->Scene->GetRenderScene();
 
-		{
-			FRDGTextureDesc OverlayColorDesc = Inputs.SceneColor.Texture->Desc;
-			OverlayColorDesc.Reset();
-			OverlayColorDesc.Format = PF_R8G8B8A8;
-			OverlayColorDesc.ClearValue = FClearValueBinding::Transparent;
-			OverlayColorDesc.Flags = TexCreate_RenderTargetable | TexCreate_ShaderResource;
-			OverlayColorDesc.NumSamples = NumSamples;
-
-			OverlayColorTexture = GraphBuilder.CreateTexture(OverlayColorDesc, TEXT("Editor.ColorOverlay"));
-		}
-		
 		if (View.ShouldRenderView() || !DepthStencilTexture)
 		{
 
@@ -144,7 +130,6 @@ FScreenPassTexture AddSelectionOutlinePass(
 
 				PassParameters->View = EditorView->GetShaderParameters();
 				PassParameters->SceneTextures = Inputs.SceneTextures;
-				PassParameters->RenderTargets[0] = FRenderTargetBinding(OverlayColorTexture, ERenderTargetLoadAction::EClear);
 				PassParameters->RenderTargets.DepthStencil = FDepthStencilBinding(
 					DepthStencilTexture,
 					ERenderTargetLoadAction::EClear,
@@ -274,8 +259,6 @@ FScreenPassTexture AddSelectionOutlinePass(
 		PassParameters->DepthSampler = PointClampSampler;
 		PassParameters->EditorPrimitivesDepth = DepthStencilTexture;
 		PassParameters->EditorPrimitivesStencil = GraphBuilder.CreateSRV(FRDGTextureSRVDesc::CreateWithPixelFormat(DepthStencilTexture, PF_X24_G8));
-		PassParameters->OverlayLookupTexture = OverlayColorTexture;
-		PassParameters->OverlayLookupSampler = PointClampSampler;
 		PassParameters->OutlineColors[0] = View.SelectionOutlineColor;
 		PassParameters->OutlineColors[1] = View.SubduedSelectionOutlineColor;
 		for (int OutlineColorIndex = 2; OutlineColorIndex < PassParameters->OutlineColors.Num(); ++OutlineColorIndex)
diff --git a/Engine/Source/Runtime/Renderer/Private/RendererScene.cpp b/Engine/Source/Runtime/Renderer/Private/RendererScene.cpp
--- a/Engine/Source/Runtime/Renderer/Private/RendererScene.cpp
+++ b/Engine/Source/Runtime/Renderer/Private/RendererScene.cpp
@@ -5957,7 +5957,7 @@ void FScene::Update(FRDGBuilder& GraphBuilder, const FUpdateParameters& Paramete
 					PrimitiveVirtualTextureLod.AddUninitialized();
 					PrimitiveOcclusionBounds.AddUninitialized();
 #if WITH_EDITOR
-					PrimitivesSelected.Add(PrimitiveSceneInfo->Proxy->WantsEditorEffects() || PrimitiveSceneInfo->Proxy->IsSelected());
+					PrimitivesSelected.Add(PrimitiveSceneInfo->Proxy->IsSelected());
 #endif
 #if RHI_RAYTRACING
 					PrimitiveRayTracingFlags.AddZeroed();
diff --git a/Engine/Source/Runtime/Renderer/Private/SceneHitProxyRendering.cpp b/Engine/Source/Runtime/Renderer/Private/SceneHitProxyRendering.cpp
--- a/Engine/Source/Runtime/Renderer/Private/SceneHitProxyRendering.cpp
+++ b/Engine/Source/Runtime/Renderer/Private/SceneHitProxyRendering.cpp
@@ -937,15 +937,11 @@ bool FEditorSelectionMeshProcessor::TryAddMeshBatch(const FMeshBatch& RESTRICT M
 
 void FEditorSelectionMeshProcessor::AddMeshBatch(const FMeshBatch& RESTRICT MeshBatch, uint64 BatchElementMask, const FPrimitiveSceneProxy* RESTRICT PrimitiveSceneProxy, int32 StaticMeshId)
 {
-	if (!PrimitiveSceneProxy)
-	{
-		return;
-	}
-	const bool bWantsEditorEffects = PrimitiveSceneProxy->WantsEditorEffects();
-	const bool bWantsOutlineForSelection = PrimitiveSceneProxy->WantsSelectionOutline() && (PrimitiveSceneProxy->IsSelected() || PrimitiveSceneProxy->IsHovered());
 	if (MeshBatch.bUseForMaterial 
 		&& MeshBatch.bUseSelectionOutline 
-		&& (bWantsEditorEffects || bWantsOutlineForSelection))
+		&& PrimitiveSceneProxy
+		&& PrimitiveSceneProxy->WantsSelectionOutline() 
+		&& (PrimitiveSceneProxy->IsSelected() || PrimitiveSceneProxy->IsHovered()))
 	{
 		const FMaterialRenderProxy* MaterialRenderProxy = MeshBatch.MaterialRenderProxy;
 		while (MaterialRenderProxy)
@@ -993,8 +989,8 @@ bool FEditorSelectionMeshProcessor::Process(
 	const int32 StencilRef = GetStencilValue(ViewIfDynamicMeshCommand, PrimitiveSceneProxy);
 	PassDrawRenderState.SetStencilRef(StencilRef);
 
-	const FHitProxyId OverlayColor = PrimitiveSceneProxy->GetOverlayColor();
-	FHitProxyShaderElementData ShaderElementData(OverlayColor);
+	FHitProxyId DummyId;
+	FHitProxyShaderElementData ShaderElementData(DummyId);
 	ShaderElementData.InitializeMeshMaterialData(ViewIfDynamicMeshCommand, PrimitiveSceneProxy, MeshBatch, StaticMeshId, false);
 
 	const FMeshDrawCommandSortKey SortKey = CalculateMeshStaticSortKey(HitProxyPassShaders.VertexShader, HitProxyPassShaders.PixelShader);
@@ -1073,7 +1069,7 @@ int32 FEditorSelectionMeshProcessor::GetStencilValue(const FSceneView* View, con
 		StencilValue = EncodeSelectionStencilValue(Color, UniqueId);
 		ProxyToStencilIndex.Add(PrimitiveSceneProxy, StencilValue);
 	}
-	else if (PrimitiveSceneProxy->IsParentSelected())
+	else
 	{
 		int Color = PrimitiveSceneProxy->GetSelectionOutlineColorIndex();
 		if (bActorSelectionColorIsSubdued && (Color == 0))
@@ -1085,10 +1081,6 @@ int32 FEditorSelectionMeshProcessor::GetStencilValue(const FSceneView* View, con
 		StencilValue = EncodeSelectionStencilValue(Color, UniqueId);
 		ActorNameToStencilIndex.Add(PrimitiveSceneProxy->GetOwnerName(), StencilValue);
 	}
-	else
-	{
-		StencilValue = ESelectionStencilValues::NotSelected;
-	}
 
 	return StencilValue;
 }
@@ -1101,7 +1093,7 @@ FEditorSelectionMeshProcessor::FEditorSelectionMeshProcessor(const FScene* Scene
 	ActorNameToStencilIndex.Add(NAME_BSP, 1);
 
 	PassDrawRenderState.SetDepthStencilState(TStaticDepthStencilState<true, CF_DepthNearOrEqual, true, CF_Always, SO_Keep, SO_Keep, SO_Replace>::GetRHI());
-	PassDrawRenderState.SetBlendState(TStaticBlendState<>::GetRHI());
+	PassDrawRenderState.SetBlendState(TStaticBlendStateWriteMask<CW_NONE, CW_NONE, CW_NONE, CW_NONE>::GetRHI());
 }
 
 FMeshPassProcessor* CreateEditorSelectionPassProcessor(ERHIFeatureLevel::Type FeatureLevel, const FScene* Scene, const FSceneView* InViewIfDynamicMeshCommand, FMeshPassDrawListContext* InDrawListContext)
diff --git a/Engine/Source/Runtime/TypedElementFramework/Public/Elements/Columns/TypedElementViewportColumns.h b/Engine/Source/Runtime/TypedElementFramework/Public/Elements/Columns/TypedElementViewportColumns.h
--- a/Engine/Source/Runtime/TypedElementFramework/Public/Elements/Columns/TypedElementViewportColumns.h
+++ b/Engine/Source/Runtime/TypedElementFramework/Public/Elements/Columns/TypedElementViewportColumns.h
@@ -8,25 +8,13 @@
 #include "TypedElementViewportColumns.generated.h"
 
 /**
- * Column to hold the color that the object is outlined with in the viewport
+ * Column to hold the color that the object is outlined with when selected in the viewport
  */
-USTRUCT(meta = (DisplayName = "Viewport Outline Color"))
-struct FTypedElementViewportOutlineColorColumn final : public FTypedElementDataStorageColumn
+USTRUCT(meta = (DisplayName = "Viewport Color"))
+struct FTypedElementViewportColorColumn final : public FTypedElementDataStorageColumn
 {
 	GENERATED_BODY()
 
 	UPROPERTY(meta = (ClampMin = "0", ClampMax = "7"))
 	uint8 SelectionOutlineColorIndex = 0;
 };
-
-/**
- * Column to hold the color that the object is overlaid with in the viewport
- */
-USTRUCT(meta = (DisplayName = "Viewport Overlay Color"))
-struct FTypedElementViewportOverlayColorColumn final : public FTypedElementDataStorageColumn
-{
-	GENERATED_BODY()
-
-	UPROPERTY()
-	FColor OverlayColor{};
-};
