Minor bug fix to the alert columns

diff --git a/Engine/Plugins/Experimental/TypedElementsDataStorage/Source/TypedElementsDataStorage/Private/Queries/TypedElementAlertQueries.cpp b/Engine/Plugins/Experimental/TypedElementsDataStorage/Source/TypedElementsDataStorage/Private/Queries/TypedElementAlertQueries.cpp
--- a/Engine/Plugins/Experimental/TypedElementsDataStorage/Source/TypedElementsDataStorage/Private/Queries/TypedElementAlertQueries.cpp
+++ b/Engine/Plugins/Experimental/TypedElementsDataStorage/Source/TypedElementsDataStorage/Private/Queries/TypedElementAlertQueries.cpp
@@ -329,6 +329,7 @@ void UTypedElementAlertQueriesFactory::RegisterOnAlertRemoveQueries(ITypedElemen
 						Alert.AlertType == FTypedElementAlertColumnType::Error ? 1 : 0,
 						Alert.AlertType == FTypedElementAlertColumnType::Warning ? 1 : 0,
 						0);
+					Alert.CachedParent = InvalidRowHandle;
 					Alert.RemoveCycleId = Context.GetUpdateCycleId();
 					Context.AddColumns<FTypedElementSyncBackToWorldTag>(Row);
 				}
@@ -377,6 +378,7 @@ void UTypedElementAlertQueriesFactory::RegisterOnParentRemoveQueries(ITypedEleme
 						0);
 					Alert.CachedParent = InvalidRowHandle;
 					Alert.RemoveCycleId = Context.GetUpdateCycleId();
+					Context.AddColumns<FTypedElementSyncBackToWorldTag>(Row);
 				}
 			})
 		.Where()
@@ -399,6 +401,7 @@ void UTypedElementAlertQueriesFactory::RegisterOnParentRemoveQueries(ITypedEleme
 						0);
 					ChildAlert.CachedParent = InvalidRowHandle;
 					ChildAlert.RemoveCycleId = Context.GetUpdateCycleId();
+					Context.AddColumns<FTypedElementSyncBackToWorldTag>(Row);
 				}
 			})
 		.Where()
