Introduced the FTypedElementSyncFromWorldInteractiveTag for multi-frame sync events.

During interactive operations, like dragging objects in the viewport, TEDS used to wait until the update was solidified. This however means that interactive tools, like the transform gizmo, would never be able to get current values from TEDS. It can also cause unintentional side effects if a sync-to-world happens, for instance when source control updates its status. This change introduces a new interactive tag that gets added by TEDS Compat whenever an interactive action starts and removed when the change is completed or canceled. Unlike the regular sync-from-world tag, this tag remains on the object during the operation resulting in selective query callbacks updating every frame. In this change list only the transform queries have been updated.

diff --git a/Engine/Plugins/Experimental/TypedElementsDataStorage/Source/TypedElementsDataStorage/Private/Compatibility/TypedElementActorTransformQueries.cpp b/Engine/Plugins/Experimental/TypedElementsDataStorage/Source/TypedElementsDataStorage/Private/Compatibility/TypedElementActorTransformQueries.cpp
--- a/Engine/Plugins/Experimental/TypedElementsDataStorage/Source/TypedElementsDataStorage/Private/Compatibility/TypedElementActorTransformQueries.cpp
+++ b/Engine/Plugins/Experimental/TypedElementsDataStorage/Source/TypedElementsDataStorage/Private/Compatibility/TypedElementActorTransformQueries.cpp
@@ -30,11 +30,12 @@ void UTypedElementActorTransformFactory::RegisterActorAddTransformColumn(ITypedE
 				const AActor* ActorInstance = Actor.Get();
 				if (ActorInstance != nullptr && ActorInstance->GetRootComponent())
 				{
-					Context.AddColumns<FTypedElementLocalTransformColumn>(Row);
+					Context.AddColumn(Row, FTypedElementLocalTransformColumn{ .Transform = ActorInstance->GetActorTransform() });
 				}
 			}
 		)
 		.Where()
+			.All<FTypedElementSyncFromWorldTag>()
 			.None<FTypedElementLocalTransformColumn>()
 		.Compile()
 	);
@@ -64,7 +65,7 @@ void UTypedElementActorTransformFactory::RegisterActorLocalTransformToColumn(ITy
 			}
 		)
 		.Where()
-			.All<FTypedElementSyncFromWorldTag>()
+			.Any<FTypedElementSyncFromWorldTag, FTypedElementSyncFromWorldInteractiveTag>()
 		.Compile()
 	);
 }
diff --git a/Engine/Plugins/Experimental/TypedElementsDataStorage/Source/TypedElementsDataStorage/Private/TypedElementDatabaseCompatibility.cpp b/Engine/Plugins/Experimental/TypedElementsDataStorage/Source/TypedElementsDataStorage/Private/TypedElementDatabaseCompatibility.cpp
--- a/Engine/Plugins/Experimental/TypedElementsDataStorage/Source/TypedElementsDataStorage/Private/TypedElementDatabaseCompatibility.cpp
+++ b/Engine/Plugins/Experimental/TypedElementsDataStorage/Source/TypedElementsDataStorage/Private/TypedElementDatabaseCompatibility.cpp
@@ -27,10 +27,12 @@ void UTypedElementDatabaseCompatibility::Initialize(ITypedElementDataStorageInte
 
 	StorageInterface->OnUpdate().AddUObject(this, &UTypedElementDatabaseCompatibility::Tick);
 
+	PreEditChangePropertyDelegateHandle = FCoreUObjectDelegates::OnPreObjectPropertyChanged.AddUObject(this, &UTypedElementDatabaseCompatibility::OnPrePropertyChanged);
 	PostEditChangePropertyDelegateHandle = FCoreUObjectDelegates::OnObjectPropertyChanged.AddUObject(this, &UTypedElementDatabaseCompatibility::OnPostEditChangeProperty);
 	ObjectModifiedDelegateHandle = FCoreUObjectDelegates::OnObjectModified.AddUObject(this, &UTypedElementDatabaseCompatibility::OnObjectModified);
 	ObjectReinstancedDelegateHandle = FCoreUObjectDelegates::OnObjectsReinstanced.AddUObject(this, &UTypedElementDatabaseCompatibility::OnObjectReinstanced);
 
+
 	PostWorldInitializationDelegateHandle = FWorldDelegates::OnPostWorldInitialization.AddUObject(this, &UTypedElementDatabaseCompatibility::OnPostWorldInitialization);
 	PreWorldFinishDestroyDelegateHandle = FWorldDelegates::OnPreWorldFinishDestroy.AddUObject(this, &UTypedElementDatabaseCompatibility::OnPreWorldFinishDestroy);
 }
@@ -48,6 +50,7 @@ void UTypedElementDatabaseCompatibility::Deinitialize()
 	FCoreUObjectDelegates::OnObjectsReinstanced.Remove(ObjectReinstancedDelegateHandle);
 	FCoreUObjectDelegates::OnObjectModified.Remove(ObjectModifiedDelegateHandle);
 	FCoreUObjectDelegates::OnObjectPropertyChanged.Remove(PostEditChangePropertyDelegateHandle);
+	FCoreUObjectDelegates::OnPreObjectPropertyChanged.Remove(PreEditChangePropertyDelegateHandle);
 	
 	Reset();
 }
@@ -632,52 +635,65 @@ void UTypedElementDatabaseCompatibility::TickPendingExternalObjectRegistration()
 
 void UTypedElementDatabaseCompatibility::TickObjectSync()
 {
-	if (!ObjectsNeedingFullSync.IsEmpty())
+	using namespace TypedElementDataStorage;
+	
+	if (!ObjectsNeedingSyncTags.IsEmpty())
 	{
-		TEDS_EVENT_SCOPE(TEXT("Process ObjectsNeedingFullSync"));
+		TEDS_EVENT_SCOPE(TEXT("Process ObjectsNeedingSyncTags"));
 		
-		TArray<TypedElementRowHandle> RowHandles;
+		using ColumnArray = TArray<const UScriptStruct*, TInlineAllocator<MaxExpectedTagsForObjectSync>>;
+		ColumnArray ColumnsToAdd;
+		ColumnArray ColumnsToRemove;
+		ColumnArray* ColumnsToAddPtr = &ColumnsToAdd;
+		ColumnArray* ColumnsToRemovePtr = &ColumnsToRemove;
+		bool bHasUpdates = false;
+		for (TPair<ObjectsNeedingSyncTagsMapKey, ObjectsNeedingSyncTagsMapValue>& ObjectToSync : ObjectsNeedingSyncTags)
 		{
-			TEDS_EVENT_SCOPE(TEXT("Reverse lookup Rows from Objects"));
-
-			RowHandles.SetNumUninitialized(ObjectsNeedingFullSync.Num());
+			const RowHandle Row = FindRowWithCompatibleObject(ObjectToSync.Key);
+			if (Storage->IsRowAvailable(Row))
 			{
-				int32 RowHandleIndex = 0;
-				for (TObjectKey<const UObject> ObjectKey : ObjectsNeedingFullSync)
+				for (FSyncTagInfo& Column : ObjectToSync.Value)
 				{
-					const TypedElementRowHandle Row = FindRowWithCompatibleObject(ObjectKey);
-					if (Row != TypedElementInvalidRowHandle)
+					if (Column.ColumnType.IsValid())
 					{
-						RowHandles[RowHandleIndex++] = Row;
+						ColumnArray* TargetColumn = Column.bAddColumn ? ColumnsToAddPtr : ColumnsToRemovePtr;
+						TargetColumn->Add(Column.ColumnType.Get());
+						bHasUpdates = true;
 					}
 				}
-				const int32 RowHandleCount = RowHandleIndex;
-				RowHandles.SetNum(RowHandleCount, EAllowShrinking::No);
+				if (bHasUpdates)
+				{
+					Storage->AddRemoveColumns(Row, ColumnsToAdd, ColumnsToRemove);
+				}
 			}
-
-			ObjectsNeedingFullSync.Reset();
+			bHasUpdates = false;
+			ColumnsToAdd.Reset();
+			ColumnsToRemove.Reset();
 		}
 
-		{
-			TEDS_EVENT_SCOPE(TEXT("Add SyncFromWorld Tag"));
-			// Tag the rows containing object data that so they get synced
-			// Note: Watch out for the performance of this, may end up doing a lot of row moves
-			for (TypedElementRowHandle Row : RowHandles)
-			{
-				Storage->AddColumn<FTypedElementSyncFromWorldTag>(Row);
-			}
-		}
+		ObjectsNeedingSyncTags.Reset();
 	}
 }
 
+void UTypedElementDatabaseCompatibility::OnPrePropertyChanged(UObject* Object, const FEditPropertyChain& PropertyChain)
+{
+	ObjectsNeedingSyncTags.FindOrAdd(Object).AddUnique(
+		FSyncTagInfo{ .ColumnType = FTypedElementSyncFromWorldInteractiveTag::StaticStruct(), .bAddColumn = true });
+}
+
 void UTypedElementDatabaseCompatibility::OnPostEditChangeProperty(
 	UObject* Object,
-	FPropertyChangedEvent& /*PropertyChangedEvent*/)
+	FPropertyChangedEvent& PropertyChangedEvent)
 {
 	// Determining the object is being tracked in the database can't be done safely as it may be queued for addition.
 	// It would also add a small bit of performance overhead as access the lookup table can be done faster as a
 	// batch operation during the tick step.
-	ObjectsNeedingFullSync.FindOrAdd(Object);
+	if (PropertyChangedEvent.ChangeType != EPropertyChangeType::Interactive)
+	{
+		ObjectsNeedingSyncTagsMapValue& SyncValue = ObjectsNeedingSyncTags.FindOrAdd(Object);
+		SyncValue.AddUnique(FSyncTagInfo{ .ColumnType = FTypedElementSyncFromWorldTag::StaticStruct(), .bAddColumn = true });
+		SyncValue.AddUnique(FSyncTagInfo{ .ColumnType = FTypedElementSyncFromWorldInteractiveTag::StaticStruct(), .bAddColumn = false });
+	}
 }
 
 void UTypedElementDatabaseCompatibility::OnObjectModified(UObject* Object)
@@ -685,7 +701,8 @@ void UTypedElementDatabaseCompatibility::OnObjectModified(UObject* Object)
 	// Determining the object is being tracked in the database can't be done safely as it may be queued for addition.
 	// It would also add a small bit of performance overhead as access the lookup table can be done faster as a
 	// batch operation during the tick step.
-	ObjectsNeedingFullSync.FindOrAdd(Object);
+	ObjectsNeedingSyncTags.FindOrAdd(Object).AddUnique(
+		FSyncTagInfo{ .ColumnType = FTypedElementSyncFromWorldTag::StaticStruct(), .bAddColumn = true });
 }
 
 void UTypedElementDatabaseCompatibility::OnObjectAdded(const void* Object, FTypedElementDatabaseCompatibilityObjectTypeInfo TypeInfo, TypedElementRowHandle Row) const
@@ -732,7 +749,10 @@ void UTypedElementDatabaseCompatibility::OnActorDestroyed(AActor* Actor)
 	RemoveCompatibleObjectExplicit(Actor);
 }
 
-
+SIZE_T GetTypeHash(const UTypedElementDatabaseCompatibility::FSyncTagInfo& Column)
+{
+	return HashCombine(Column.ColumnType.GetWeakPtrTypeHash(), Column.bAddColumn);
+}
 
 //
 // UTypedElementDatabaseCompatibility::FRegistrationCommandChange
diff --git a/Engine/Plugins/Experimental/TypedElementsDataStorage/Source/TypedElementsDataStorage/Private/TypedElementDatabaseCompatibility.h b/Engine/Plugins/Experimental/TypedElementsDataStorage/Source/TypedElementsDataStorage/Private/TypedElementDatabaseCompatibility.h
--- a/Engine/Plugins/Experimental/TypedElementsDataStorage/Source/TypedElementsDataStorage/Private/TypedElementDatabaseCompatibility.h
+++ b/Engine/Plugins/Experimental/TypedElementsDataStorage/Source/TypedElementsDataStorage/Private/TypedElementDatabaseCompatibility.h
@@ -103,6 +103,7 @@ private:
 	void TickPendingExternalObjectRegistration();
 	void TickObjectSync();
 
+	void OnPrePropertyChanged(UObject* Object, const FEditPropertyChain& PropertyChain);
 	void OnPostEditChangeProperty(UObject* Object, FPropertyChangedEvent& PropertyChangedEvent);
 	void OnObjectModified(UObject* Object);
 	void OnObjectAdded(const void* Object, FTypedElementDatabaseCompatibilityObjectTypeInfo TypeInfo, TypedElementRowHandle Row) const;
@@ -189,9 +190,23 @@ private:
 	 * Reference of objects (UObject and AActor) that need to be fully synced from the world to the database.
 	 * Caution: Could point to objects that have been GC-ed
 	 */
-	TSet<TObjectKey<const UObject>> ObjectsNeedingFullSync;
+	struct FSyncTagInfo
+	{
+		TWeakObjectPtr<const UScriptStruct> ColumnType;
+		bool bAddColumn;
+
+		bool operator==(const FSyncTagInfo& Rhs) const = default;
+		bool operator!=(const FSyncTagInfo& Rhs) const = default;
+	};
+	friend SIZE_T GetTypeHash(const FSyncTagInfo& Column);
+	static constexpr uint32 MaxExpectedTagsForObjectSync = 2;
+	using ObjectsNeedingSyncTagsMapKey = TObjectKey<const UObject>;
+	using ObjectsNeedingSyncTagsMapValue = TArray<FSyncTagInfo, TInlineAllocator<MaxExpectedTagsForObjectSync>>;
+	using ObjectsNeedingSyncTagsMap = TMap<ObjectsNeedingSyncTagsMapKey, ObjectsNeedingSyncTagsMapValue>;
+	ObjectsNeedingSyncTagsMap ObjectsNeedingSyncTags;
 
 	TMap<UWorld*, FDelegateHandle> ActorDestroyedDelegateHandles;
+	FDelegateHandle PreEditChangePropertyDelegateHandle;
 	FDelegateHandle PostEditChangePropertyDelegateHandle;
 	FDelegateHandle ObjectModifiedDelegateHandle;
 	FDelegateHandle PostWorldInitializationDelegateHandle;
@@ -202,6 +217,8 @@ private:
 	TypedElementDataStorage::QueryHandle ScriptStructTypeInfoQuery;
 };
 
+SIZE_T GetTypeHash(const UTypedElementDatabaseCompatibility::FSyncTagInfo& Column);
+
 enum class ETypedElementDatabaseCompatibilityObjectType : uint8
 {
 	Struct,
diff --git a/Engine/Source/Runtime/TypedElementFramework/Public/Elements/Columns/TypedElementMiscColumns.h b/Engine/Source/Runtime/TypedElementFramework/Public/Elements/Columns/TypedElementMiscColumns.h
--- a/Engine/Source/Runtime/TypedElementFramework/Public/Elements/Columns/TypedElementMiscColumns.h
+++ b/Engine/Source/Runtime/TypedElementFramework/Public/Elements/Columns/TypedElementMiscColumns.h
@@ -20,7 +20,7 @@ struct FTypedElementSyncBackToWorldTag final : public FTypedElementDataStorageTa
 };
 
 /**
- * Tag to signal that data a processor copies out of the world must be synced to the datastorage.
+ * Tag to signal that data a processor copies out of the world must be synced to the data storage.
  * Useful for when an Actor was recently spawned or reloaded in the world.
  * Currently used if any property changes since there is no mechanism to selectively run
  * queries for specific changed properties.
@@ -31,10 +31,27 @@ struct FTypedElementSyncFromWorldTag final : public FTypedElementDataStorageTag
 	GENERATED_BODY()
 };
 
+/**
+ * Tag to signal that data a processor copies out of the world must be synced to the data storage.
+ * Useful for when an Actor was recently spawned or reloaded in the world. This version is not
+ * automatically removed and intended for interactive operations that will take a few frames
+ * to complete such as dragging.
+ * Currently used if any property changes since there is no mechanism to selectively run
+ * queries for specific changed properties.
+ */
+USTRUCT(meta = (DisplayName = "Sync from world (interactive)"))
+struct FTypedElementSyncFromWorldInteractiveTag final : public FTypedElementDataStorageTag
+{
+	GENERATED_BODY()
+};
+
+/**
+ * A general reference to another row. 
+ */
 USTRUCT(meta = (DisplayName = "Row reference"))
 struct FTypedElementRowReferenceColumn final : public FTypedElementDataStorageColumn
 {
 	GENERATED_BODY()
 
-	TypedElementRowHandle Row;
+	TypedElementDataStorage::RowHandle Row;
 };
\ No newline at end of file
