Columns created in TEDS that have intialization data assigned now have that data applied before the OnAdd observer callback is triggered.

Even with the recent change to Mass, columns weren't being assigned their data before OnAdd observers were triggered. This however was a problem on the TEDS side as registring columns happens in a "add" and separate "set" step. These have now be combined togehter into one so the initialization happens before OnAdd is called. For this to work however Mass has gotten an additional function that calls an initialization function if provided, to get the timing correct.

diff --git a/Engine/Plugins/Experimental/TypedElementsDataStorage/Source/TypedElementsDataStorage/Private/Processors/TypedElementProcessorAdaptors.cpp b/Engine/Plugins/Experimental/TypedElementsDataStorage/Source/TypedElementsDataStorage/Private/Processors/TypedElementProcessorAdaptors.cpp
--- a/Engine/Plugins/Experimental/TypedElementsDataStorage/Source/TypedElementsDataStorage/Private/Processors/TypedElementProcessorAdaptors.cpp
+++ b/Engine/Plugins/Experimental/TypedElementsDataStorage/Source/TypedElementsDataStorage/Private/Processors/TypedElementProcessorAdaptors.cpp
@@ -176,9 +176,9 @@ public:
 	void* AddColumnUninitialized(TypedElementDataStorage::RowHandle Row, const UScriptStruct* ObjectType) override
 	{
 		return AddColumnUninitialized(Row, ObjectType,
-			[](const UScriptStruct* TypeInfo, void* Destination, void* Source)
+			[](const UScriptStruct& TypeInfo, void* Destination, void* Source)
 			{
-				TypeInfo->CopyScriptStruct(Destination, Source);
+				TypeInfo.CopyScriptStruct(Destination, Source);
 			});
 	}
 
@@ -226,19 +226,16 @@ public:
 					FStructView Fragment = System.GetFragmentDataStruct(AddedColumn->Entity, AddedColumn->FragmentType);
 					if (!Fragment.IsValid())
 					{
-						System.AddFragmentToEntity(AddedColumn->Entity, AddedColumn->FragmentType);
+						System.AddFragmentToEntity(AddedColumn->Entity, AddedColumn->FragmentType, 
+							[AddedColumn](void* Fragment, const UScriptStruct& FragmentType)
+							{
+								AddedColumn->Relocator(FragmentType, Fragment, AddedColumn->Object);
+							});
+					}
+					else
+					{
+						AddedColumn->Relocator(*AddedColumn->FragmentType, Fragment.GetMemory(), AddedColumn->Object);
 					}
-				}
-			});
-
-		this->Context.Defer().template PushCommand<FMassDeferredSetCommand>(
-			[AddedColumn](FMassEntityManager& System)
-			{
-				// Check entity before proceeding. It's possible it may have been invalidated before this deferred call fired.
-				if (System.IsEntityActive(AddedColumn->Entity))
-				{
-					FStructView Fragment = System.GetFragmentDataStruct(AddedColumn->Entity, AddedColumn->FragmentType);
-					AddedColumn->Relocator(AddedColumn->FragmentType, Fragment.GetMemory(), AddedColumn->Object);
 				}
 			});
 		
diff --git a/Engine/Plugins/Runtime/MassEntity/Source/MassEntity/Private/MassEntityManager.cpp b/Engine/Plugins/Runtime/MassEntity/Source/MassEntity/Private/MassEntityManager.cpp
--- a/Engine/Plugins/Runtime/MassEntity/Source/MassEntity/Private/MassEntityManager.cpp
+++ b/Engine/Plugins/Runtime/MassEntity/Source/MassEntity/Private/MassEntityManager.cpp
@@ -718,6 +718,22 @@ void FMassEntityManager::AddFragmentToEntity(FMassEntityHandle Entity, const USc
 	ObserverManager.OnPostCompositionAdded(Entity, Descriptor);
 }
 
+void FMassEntityManager::AddFragmentToEntity(FMassEntityHandle Entity, const UScriptStruct* FragmentType, const FStructInitializationCallback& Initializer)
+{
+	checkf(FragmentType, TEXT("Null fragment type passed in to %hs"), __FUNCTION__);
+	checkf(IsProcessing() == false, TEXT("Synchronous API function %hs called during mass processing. Use asynchronous API instead."), __FUNCTION__);
+
+	CheckIfEntityIsActive(Entity);
+
+	FMassFragmentBitSet Fragments = InternalAddFragmentListToEntityChecked(Entity, FMassFragmentBitSet(*FragmentType));
+	FEntityData& EntityData = Entities[Entity.Index];
+	void* FragmentData = EntityData.CurrentArchetype->GetFragmentDataForEntity(FragmentType, Entity.Index);
+	Initializer(FragmentData, *FragmentType);
+
+	const FMassArchetypeCompositionDescriptor Descriptor(MoveTemp(Fragments));
+	ObserverManager.OnPostCompositionAdded(Entity, Descriptor);
+}
+
 void FMassEntityManager::AddFragmentListToEntity(FMassEntityHandle Entity, TConstArrayView<const UScriptStruct*> FragmentList)
 {
 	CheckIfEntityIsActive(Entity);
diff --git a/Engine/Plugins/Runtime/MassEntity/Source/MassEntity/Public/MassEntityManager.h b/Engine/Plugins/Runtime/MassEntity/Source/MassEntity/Public/MassEntityManager.h
--- a/Engine/Plugins/Runtime/MassEntity/Source/MassEntity/Public/MassEntityManager.h
+++ b/Engine/Plugins/Runtime/MassEntity/Source/MassEntity/Public/MassEntityManager.h
@@ -11,6 +11,7 @@
 #include "MassObserverManager.h"
 #include "Containers/MpscQueue.h"
 #include "MassRequirementAccessDetector.h"
+#include "Templates/FunctionFwd.h"
 
 
 struct FMassEntityQuery;
@@ -78,6 +79,7 @@ public:
 	private:
 		std::atomic<int32>& ScopedProcessingCount;
 	};
+	using FStructInitializationCallback = TFunctionRef<void(void* Fragment, const UScriptStruct& FragmentType)>;
 
 	const static FMassEntityHandle InvalidEntity;
 
@@ -250,6 +252,7 @@ public:
 	void BatchDestroyEntityChunks(const FMassArchetypeEntityCollection& Collection);
 
 	void AddFragmentToEntity(FMassEntityHandle Entity, const UScriptStruct* FragmentType);
+	void AddFragmentToEntity(FMassEntityHandle Entity, const UScriptStruct* FragmentType, const FStructInitializationCallback& Initializer);
 
 	/** 
 	 *  Ensures that only unique fragments are added. 
diff --git a/Engine/Source/Runtime/TypedElementFramework/Public/Elements/Interfaces/TypedElementQueryStorageInterfaces.h b/Engine/Source/Runtime/TypedElementFramework/Public/Elements/Interfaces/TypedElementQueryStorageInterfaces.h
--- a/Engine/Source/Runtime/TypedElementFramework/Public/Elements/Interfaces/TypedElementQueryStorageInterfaces.h
+++ b/Engine/Source/Runtime/TypedElementFramework/Public/Elements/Interfaces/TypedElementQueryStorageInterfaces.h
@@ -92,7 +92,7 @@ namespace TypedElementDataStorage
 
 	struct ICommonQueryWithEnvironmentContext : public ICommonQueryContext
 	{
-		using ObjectCopyOrMove = void (*)(const UScriptStruct* TypeInfo, void* Destination, void* Source);
+		using ObjectCopyOrMove = void (*)(const UScriptStruct& TypeInfo, void* Destination, void* Source);
 
 		/**
 		 * Returns the id for the current update cycle. Every time TEDS goes through a cycle of running query callbacks, this is
@@ -304,7 +304,7 @@ namespace TypedElementDataStorage
 		if constexpr (std::is_move_constructible_v<ColumnType>)
 		{
 			void* Address = AddColumnUninitialized(Row, TypeInfo,
-				[](const UScriptStruct*, void* Destination, void* Source)
+				[](const UScriptStruct&, void* Destination, void* Source)
 				{
 					*reinterpret_cast<ColumnType*>(Destination) = MoveTemp(*reinterpret_cast<ColumnType*>(Source));
 				});
