Made sure commands from a single row sub-query in TEDS doesn't flush the command buffer.

diff --git a/Engine/Plugins/Experimental/TypedElementsDataStorage/Source/TypedElementsDataStorage/Private/Processors/TypedElementProcessorAdaptors.cpp b/Engine/Plugins/Experimental/TypedElementsDataStorage/Source/TypedElementsDataStorage/Private/Processors/TypedElementProcessorAdaptors.cpp
--- a/Engine/Plugins/Experimental/TypedElementsDataStorage/Source/TypedElementsDataStorage/Private/Processors/TypedElementProcessorAdaptors.cpp
+++ b/Engine/Plugins/Experimental/TypedElementsDataStorage/Source/TypedElementsDataStorage/Private/Processors/TypedElementProcessorAdaptors.cpp
@@ -728,6 +728,7 @@ TypedElementDataStorage::FQueryResult FTypedElementQueryProcessorData::Execute(
 		FMassExecutionContext Context(EntityManager);
 		Context.SetEntityCollection(FMassArchetypeEntityCollection(NativeArchetype, { NativeEntity }, FMassArchetypeEntityCollection::NoDuplicates));
 		Context.SetDeferredCommandBuffer(ParentContext.GetSharedDeferredCommandBuffer());
+		Context.SetFlushDeferredCommands(false);
 
 		NativeQuery.ForEachEntityChunk(EntityManager, Context,
 			[&Result, &Callback, &Description, &Environment, &EntityManager, RowHandle](FMassExecutionContext& Context)
