TEDS: Fix for crash on start up with compiler weirdness around execution ordering.

diff --git a/Engine/Plugins/Experimental/TypedElementsDataStorage/Source/TypedElementsDataStorage/Private/TypedElementDatabaseUI.cpp b/Engine/Plugins/Experimental/TypedElementsDataStorage/Source/TypedElementsDataStorage/Private/TypedElementDatabaseUI.cpp
--- a/Engine/Plugins/Experimental/TypedElementsDataStorage/Source/TypedElementsDataStorage/Private/TypedElementDatabaseUI.cpp
+++ b/Engine/Plugins/Experimental/TypedElementsDataStorage/Source/TypedElementsDataStorage/Private/TypedElementDatabaseUI.cpp
@@ -391,7 +391,8 @@ bool UTypedElementDatabaseUi::CreateSingleWidgetConstructor(
 			{
 				Target->InitializeStruct(Result.Get());
 				Result->Initialize(Arguments, MoveTemp(MatchedColumnTypes), QueryConditions);
-				return Callback(MoveTemp(Result), Result->GetMatchedColumns());
+				const TArray<TWeakObjectPtr<const UScriptStruct>>& MatchedColumns = Result->GetMatchedColumns();
+				return Callback(MoveTemp(Result), MatchedColumns);
 			}
 			return true;
 		}
