Activatable query callbacks.

With this update it's possible to make a query callback activatable based on a name. From a query callback or directly from TEDS it's possible to enable these queries to run for one update cycle, enabling at the start of the next update cycle and disabling at the end of that cycle.
The following changes were also made:
- Moved the query store to the environment so it's more widely available.
- Made the default tick group explicitly called "Default" to make sure dependencies are correctly respected when there's a dependency on the default tick group.
- Added additional tick groups "PreUpdate", "Update" and "PostUpdate".
- Fixed an issue that caused tick groups to not be fully registered, causing dependencies to break.
- Scratch buffer is no longer stored in a shared pointer as this was no longer needed after the latest scratch buffer updates.
- Removed  a portion of the original code that was intended to deal with boolean operations. This has been implemented in the query conditions.

diff --git a/Engine/Plugins/Experimental/TypedElementsDataStorage/Source/TypedElementsDataStorage/Private/Processors/TypedElementProcessorAdaptors.cpp b/Engine/Plugins/Experimental/TypedElementsDataStorage/Source/TypedElementsDataStorage/Private/Processors/TypedElementProcessorAdaptors.cpp
--- a/Engine/Plugins/Experimental/TypedElementsDataStorage/Source/TypedElementsDataStorage/Private/Processors/TypedElementProcessorAdaptors.cpp
+++ b/Engine/Plugins/Experimental/TypedElementsDataStorage/Source/TypedElementsDataStorage/Private/Processors/TypedElementProcessorAdaptors.cpp
@@ -173,6 +173,15 @@ public:
 		return Environment.GetUpdateCycleId();
 	}
 
+	void ActivateQueries(FName ActivationName)
+	{
+		this->Context.Defer().template PushCommand<FMassDeferredCommand<EMassCommandOperationType::None>>(
+			[Environment = &this->Environment, ActivationName](FMassEntityManager&)
+			{
+				Environment->GetQueryStore().ActivateQueries(ActivationName);
+			});
+	}
+
 	void* AddColumnUninitialized(TypedElementDataStorage::RowHandle Row, const UScriptStruct* ObjectType) override
 	{
 		return AddColumnUninitialized(Row, ObjectType,
@@ -538,16 +547,19 @@ void FPhasePreOrPostAmbleExecutor::ExecuteQuery(
 	FMassEntityQuery& NativeQuery,
 	ITypedElementDataStorageInterface::QueryCallbackRef Callback)
 {
-	NativeQuery.ForEachEntityChunk(Context.GetEntityManagerChecked(), Context,
-		[&Callback, &QueryStore, &Environment, &Description](FMassExecutionContext& ExecutionContext)
-		{
-			if (FTypedElementQueryProcessorData::PrepareCachedDependenciesOnQuery(Description, ExecutionContext))
+	if (Description.Callback.ActivationCount > 0)
+	{
+		NativeQuery.ForEachEntityChunk(Context.GetEntityManagerChecked(), Context,
+			[&Callback, &QueryStore, &Environment, &Description](FMassExecutionContext& ExecutionContext)
 			{
-				FMassContextForwarder QueryContext(Description, ExecutionContext, QueryStore, Environment);
-				Callback(Description, QueryContext);
+				if (FTypedElementQueryProcessorData::PrepareCachedDependenciesOnQuery(Description, ExecutionContext))
+				{
+					FMassContextForwarder QueryContext(Description, ExecutionContext, QueryStore, Environment);
+					Callback(Description, QueryContext);
+				}
 			}
-		}
-	);
+		);
+	}
 }
 
 
@@ -624,8 +636,22 @@ EMassProcessingPhase FTypedElementQueryProcessorData::MapToMassProcessingPhase(I
 
 FString FTypedElementQueryProcessorData::GetProcessorName() const
 {
-	const FTypedElementExtendedQuery* StoredQuery = QueryStore ? QueryStore->Get(ParentQuery) : nullptr;
-	return StoredQuery ? StoredQuery->Description.Callback.Name.ToString() : FString(TEXT("<unnamed>"));
+	FString Result;
+	if (const FTypedElementExtendedQuery* StoredQuery = QueryStore ? QueryStore->Get(ParentQuery) : nullptr)
+	{
+		Result = StoredQuery->Description.Callback.Name.ToString();
+		if (!StoredQuery->Description.Callback.ActivationName.IsNone())
+		{
+			Result += TEXT(" (Activatable: '");
+			StoredQuery->Description.Callback.ActivationName.AppendString(Result);
+			Result += TEXT("')");
+		}
+	}
+	else
+	{
+		Result = TEXT("<unnamed>");
+	}
+	return Result;
 }
 
 bool FTypedElementQueryProcessorData::PrepareCachedDependenciesOnQuery(
@@ -668,19 +694,23 @@ TypedElementDataStorage::FQueryResult FTypedElementQueryProcessorData::Execute(
 	FMassEntityManager& EntityManager,
 	FTypedElementDatabaseEnvironment& Environment)
 {
-	FMassExecutionContext Context(EntityManager);
 	ITypedElementDataStorageInterface::FQueryResult Result;
 	Result.Completed = ITypedElementDataStorageInterface::FQueryResult::ECompletion::Fully;
 	
-	NativeQuery.ForEachEntityChunk(EntityManager, Context,
-		[&Result, &Callback, &Description](FMassExecutionContext& Context)
-		{
-			// No need to cache any subsystem dependencies as these are not accessible from a direct query.
-			FMassDirectContextForwarder QueryContext(Context);
-			Callback(Description, QueryContext);
-			Result.Count += Context.GetNumEntities();
-		}
-	);
+	if (Description.Callback.ActivationCount > 0)
+	{
+		FMassExecutionContext Context(EntityManager);
+		
+		NativeQuery.ForEachEntityChunk(EntityManager, Context,
+			[&Result, &Callback, &Description](FMassExecutionContext& Context)
+			{
+				// No need to cache any subsystem dependencies as these are not accessible from a direct query.
+				FMassDirectContextForwarder QueryContext(Context);
+				Callback(Description, QueryContext);
+				Result.Count += Context.GetNumEntities();
+			}
+		);
+	}
 	return Result;
 }
 
@@ -692,20 +722,25 @@ TypedElementDataStorage::FQueryResult FTypedElementQueryProcessorData::Execute(
 	FTypedElementDatabaseEnvironment& Environment,
 	FMassExecutionContext& ParentContext)
 {
-	FMassExecutionContext Context(EntityManager);
-	Context.SetDeferredCommandBuffer(ParentContext.GetSharedDeferredCommandBuffer());
 	ITypedElementDataStorageInterface::FQueryResult Result;
 	Result.Completed = ITypedElementDataStorageInterface::FQueryResult::ECompletion::Fully;
 
-	NativeQuery.ForEachEntityChunk(EntityManager, Context,
-		[&Result, &Callback, &Description, &Environment](FMassExecutionContext& Context)
-		{
-			// No need to cache any subsystem dependencies as these are not accessible from a subquery.
-			FMassSubqueryContextForwarder QueryContext(Context, Environment);
-			Callback(Description, QueryContext);
-			Result.Count += Context.GetNumEntities();
-		}
-	);
+	if (Description.Callback.ActivationCount > 0)
+	{
+		FMassExecutionContext Context(EntityManager);
+		Context.SetDeferredCommandBuffer(ParentContext.GetSharedDeferredCommandBuffer());
+		Context.SetFlushDeferredCommands(false);
+
+		NativeQuery.ForEachEntityChunk(EntityManager, Context,
+			[&Result, &Callback, &Description, &Environment](FMassExecutionContext& Context)
+			{
+				// No need to cache any subsystem dependencies as these are not accessible from a subquery.
+				FMassSubqueryContextForwarder QueryContext(Context, Environment);
+				Callback(Description, QueryContext);
+				Result.Count += Context.GetNumEntities();
+			}
+		);
+	}
 	return Result;
 }
 
@@ -722,7 +757,7 @@ TypedElementDataStorage::FQueryResult FTypedElementQueryProcessorData::Execute(
 	Result.Completed = ITypedElementDataStorageInterface::FQueryResult::ECompletion::Fully;
 
 	FMassEntityHandle NativeEntity = FMassEntityHandle::FromNumber(RowHandle);
-	if (EntityManager.IsEntityActive(NativeEntity))
+	if (Description.Callback.ActivationCount > 0 && EntityManager.IsEntityActive(NativeEntity))
 	{
 		FMassArchetypeHandle NativeArchetype = EntityManager.GetArchetypeForEntityUnsafe(NativeEntity);
 		FMassExecutionContext Context(EntityManager);
@@ -750,16 +785,19 @@ void FTypedElementQueryProcessorData::Execute(FMassEntityManager& EntityManager,
 	checkf(StoredQuery, TEXT("A query callback was registered for execution without an associated query."));
 	
 	ITypedElementDataStorageInterface::FQueryDescription& Description = StoredQuery->Description;
-	NativeQuery.ForEachEntityChunk(EntityManager, Context,
-		[this, &Description](FMassExecutionContext& Context)
-		{
-			if (PrepareCachedDependenciesOnQuery(Description, Context))
+	if (Description.Callback.ActivationCount > 0)
+	{
+		NativeQuery.ForEachEntityChunk(EntityManager, Context,
+			[this, &Description](FMassExecutionContext& Context)
 			{
-				FMassContextForwarder QueryContext(Description, Context, *QueryStore, *Environment);
-				Description.Callback.Function(Description, QueryContext);
+				if (PrepareCachedDependenciesOnQuery(Description, Context))
+				{
+					FMassContextForwarder QueryContext(Description, Context, *QueryStore, *Environment);
+					Description.Callback.Function(Description, QueryContext);
+				}
 			}
-		}
-	);
+		);
+	}
 }
 
 
diff --git a/Engine/Plugins/Experimental/TypedElementsDataStorage/Source/TypedElementsDataStorage/Private/Queries/TypedElementExtendedQueryStore.cpp b/Engine/Plugins/Experimental/TypedElementsDataStorage/Source/TypedElementsDataStorage/Private/Queries/TypedElementExtendedQueryStore.cpp
--- a/Engine/Plugins/Experimental/TypedElementsDataStorage/Source/TypedElementsDataStorage/Private/Queries/TypedElementExtendedQueryStore.cpp
+++ b/Engine/Plugins/Experimental/TypedElementsDataStorage/Source/TypedElementsDataStorage/Private/Queries/TypedElementExtendedQueryStore.cpp
@@ -25,6 +25,7 @@ FTypedElementExtendedQueryStore::Handle FTypedElementExtendedQueryStore::Registe
 	bContinueSetup = bContinueSetup &&	SetupDependencies(StoredQuery.Description, NativeQuery);
 	bContinueSetup = bContinueSetup &&	SetupTickGroupDefaults(StoredQuery.Description);
 	bContinueSetup = bContinueSetup &&	SetupProcessors(Result, StoredQuery, Environment, EntityManager, PhaseManager);
+	bContinueSetup = bContinueSetup &&	SetupActivatable(Result, Query);
 	
 	if (!bContinueSetup)
 	{
@@ -64,12 +65,12 @@ void FTypedElementExtendedQueryStore::RegisterTickGroup(FName GroupName, ITypedE
 {
 	FTickGroupDescription& Group = TickGroupDescriptions.FindOrAdd({ GroupName, Phase });
 
-	if (!Group.BeforeGroups.Find(BeforeGroup))
+	if (!BeforeGroup.IsNone() && Group.BeforeGroups.Find(BeforeGroup) == INDEX_NONE)
 	{
 		Group.BeforeGroups.Add(BeforeGroup);
 	}
 
-	if (!Group.AfterGroups.Find(AfterGroup))
+	if (!AfterGroup.IsNone() && Group.AfterGroups.Find(AfterGroup) == INDEX_NONE)
 	{
 		Group.AfterGroups.Add(AfterGroup);
 	}
@@ -131,6 +132,49 @@ void FTypedElementExtendedQueryStore::ListAliveEntries(const ListAliveEntriesCon
 	Queries.ListAliveEntries(Callback);
 }
 
+void FTypedElementExtendedQueryStore::UpdateActivatableQueries()
+{
+	// Update activatable counts and remove any queries that have completed.
+	for (Handle Query : ActiveActivatables)
+	{
+		FTypedElementExtendedQuery& QueryData = Queries.GetMutable(Query);
+		checkf(QueryData.Description.Callback.ActivationCount > 0,
+			TEXT("Attempting to decrement the query '%s' which is already at zero."), *QueryData.Description.Callback.Name.ToString());
+		QueryData.Description.Callback.ActivationCount--;
+	}
+	ActiveActivatables.Reset();
+
+	// Queue up the next batch of activatables.
+	for (Handle Query : PendingActivatables)
+	{
+		FTypedElementExtendedQuery& QueryData = Queries.GetMutable(Query);
+		if (QueryData.Description.Callback.ActivationCount == 0)
+		{
+			QueryData.Description.Callback.ActivationCount = 1;
+			ActiveActivatables.Add(Query);
+		}
+	}
+	PendingActivatables.Reset();
+}
+
+void FTypedElementExtendedQueryStore::ActivateQueries(FName ActivationName)
+{
+	for (TMultiMap<FName, Handle>::TKeyIterator QueryIt = ActivatableMapping.CreateKeyIterator(ActivationName); QueryIt; ++QueryIt)
+	{
+		Handle Query = QueryIt.Value();
+		if (Queries.IsAlive(Query))
+		{
+#if DO_ENSURE
+			FTypedElementExtendedQuery& QueryData = Queries.GetMutable(Query);
+			checkf(!QueryData.Description.Callback.ActivationName.IsNone(),
+				TEXT("Attempting to enable the query '%s' which isn't activatable."), *QueryData.Description.Callback.Name.ToString());
+#endif
+			PendingActivatables.Add(Query);
+		}
+	}
+}
+
+
 TypedElementDataStorage::FQueryResult FTypedElementExtendedQueryStore::RunQuery(FMassEntityManager& EntityManager, Handle Query)
 {
 	using ActionType = ITypedElementDataStorageInterface::FQueryDescription::EActionType;
@@ -637,6 +681,15 @@ bool FTypedElementExtendedQueryStore::SetupProcessors(Handle QueryHandle, FTyped
 	return true;
 }
 
+bool FTypedElementExtendedQueryStore::SetupActivatable(Handle QueryHandle, ITypedElementDataStorageInterface::FQueryDescription& Query)
+{
+	if (!Query.Callback.ActivationName.IsNone())
+	{
+		ActivatableMapping.Add(Query.Callback.ActivationName, QueryHandle);
+	}
+	return true;
+}
+
 EMassFragmentAccess FTypedElementExtendedQueryStore::ConvertToNativeAccessType(ITypedElementDataStorageInterface::EQueryAccessType AccessType)
 {
 	switch (AccessType)
@@ -696,6 +749,13 @@ void FTypedElementExtendedQueryStore::RunPhasePreOrPostAmbleQueries(FMassEntityM
 
 void FTypedElementExtendedQueryStore::UnregisterQueryData(Handle Query, FTypedElementExtendedQuery& QueryData, FMassProcessingPhaseManager& PhaseManager)
 {
+	if (!QueryData.Description.Callback.ActivationName.IsNone())
+	{
+		ActivatableMapping.RemoveSingle(QueryData.Description.Callback.ActivationName, Query);
+		ActiveActivatables.RemoveSingleSwap(Query);
+		PendingActivatables.RemoveSingleSwap(Query);
+	}
+
 	if (QueryData.Processor)
 	{
 		if (QueryData.Processor->IsA<UTypedElementQueryProcessorCallbackAdapterProcessorBase>())
diff --git a/Engine/Plugins/Experimental/TypedElementsDataStorage/Source/TypedElementsDataStorage/Private/Queries/TypedElementExtendedQueryStore.h b/Engine/Plugins/Experimental/TypedElementsDataStorage/Source/TypedElementsDataStorage/Private/Queries/TypedElementExtendedQueryStore.h
--- a/Engine/Plugins/Experimental/TypedElementsDataStorage/Source/TypedElementsDataStorage/Private/Queries/TypedElementExtendedQueryStore.h
+++ b/Engine/Plugins/Experimental/TypedElementsDataStorage/Source/TypedElementsDataStorage/Private/Queries/TypedElementExtendedQueryStore.h
@@ -83,10 +83,21 @@ public:
 	/** Calls the provided callback for each query that's available. */
 	void ListAliveEntries(const ListAliveEntriesConstCallback& Callback) const;
 
+	/**
+	 * @section activatable queries
+	 * @description Functions to manipulate activatable queries
+	 */
+
+	/** Update the active activatable queries. In practice this means decrementing any active queries that automatically decrement. */
+	void UpdateActivatableQueries();
+	/** Triggers a query to run for a single update cycle. */
+	void ActivateQueries(FName ActivationName);
+	
 	/**
 	 * @section Execution
 	 * @description Various functions to run queries.
 	 */
+
 	TypedElementDataStorage::FQueryResult RunQuery(FMassEntityManager& EntityManager, Handle Query);
 	TypedElementDataStorage::FQueryResult RunQuery(
 		FMassEntityManager& EntityManager, 
@@ -155,6 +166,7 @@ private:
 	bool SetupTickGroupDefaults(ITypedElementDataStorageInterface::FQueryDescription& Query);
 	bool SetupProcessors(Handle QueryHandle, FTypedElementExtendedQuery& StoredQuery, FTypedElementDatabaseEnvironment& Environment,
 		FMassEntityManager& EntityManager, FMassProcessingPhaseManager& PhaseManager);
+	bool SetupActivatable(Handle QueryHandle, ITypedElementDataStorageInterface::FQueryDescription& Query);
 
 	EMassFragmentAccess ConvertToNativeAccessType(ITypedElementDataStorageInterface::EQueryAccessType AccessType);
 
@@ -170,7 +182,10 @@ private:
 	static const ITypedElementDataStorageInterface::FQueryDescription EmptyDescription;
 
 	QueryStore Queries;
+	TMultiMap<FName, Handle> ActivatableMapping;
 	TMap<FTickGroupId, FTickGroupDescription> TickGroupDescriptions;
 	TArray<Handle> PhasePreparationQueries[MaxTickPhase];
 	TArray<Handle> PhaseFinalizationQueries[MaxTickPhase];
+	TArray<Handle> PendingActivatables;
+	TArray<Handle> ActiveActivatables;
 };
diff --git a/Engine/Plugins/Experimental/TypedElementsDataStorage/Source/TypedElementsDataStorage/Private/TypedElementDatabase.cpp b/Engine/Plugins/Experimental/TypedElementsDataStorage/Source/TypedElementsDataStorage/Private/TypedElementDatabase.cpp
--- a/Engine/Plugins/Experimental/TypedElementsDataStorage/Source/TypedElementsDataStorage/Private/TypedElementDatabase.cpp
+++ b/Engine/Plugins/Experimental/TypedElementsDataStorage/Source/TypedElementsDataStorage/Private/TypedElementDatabase.cpp
@@ -18,6 +18,10 @@
 #include "TickTaskManagerInterface.h"
 #include "UObject/UObjectIterator.h"
 
+const FName UTypedElementDatabase::TickGroupName_Default(TEXT("Default"));
+const FName UTypedElementDatabase::TickGroupName_PreUpdate(TEXT("PreUpdate"));
+const FName UTypedElementDatabase::TickGroupName_Update(TEXT("Update"));
+const FName UTypedElementDatabase::TickGroupName_PostUpdate(TEXT("PostUpdate"));
 const FName UTypedElementDatabase::TickGroupName_SyncWidget(TEXT("SyncWidgets"));
 
 FAutoConsoleCommandWithOutputDevice PrintQueryCallbacksConsoleCommand(
@@ -127,15 +131,29 @@ void UTypedElementDatabase::Initialize()
 					FinalizePhase(Phase, DeltaTime);
 				});
 
-			// Guarantee that syncing to the data storage always happens before syncing to external.
+			// Update external source to TEDS at the start of the phase.
 			RegisterTickGroup(GetQueryTickGroupName(EQueryTickGroups::SyncExternalToDataStorage),
-				Phase, GetQueryTickGroupName(EQueryTickGroups::SyncDataStorageToExternal), {}, false);
-			// Guarantee that widgets syncs happen after external data has been updated and internal processing 
-			// completed to the data storage.
-			RegisterTickGroup(GetQueryTickGroupName(EQueryTickGroups::SyncWidgets),
-				Phase, {}, GetQueryTickGroupName(EQueryTickGroups::Default), true /* Needs main thread*/);
+				Phase, {}, {}, false);
+			
+			// Default group.
+			RegisterTickGroup(GetQueryTickGroupName(EQueryTickGroups::Default),
+				Phase, {}, GetQueryTickGroupName(EQueryTickGroups::SyncExternalToDataStorage), false);
+
+			// Order the update groups.
+			RegisterTickGroup(GetQueryTickGroupName(EQueryTickGroups::PreUpdate),
+				Phase, {}, GetQueryTickGroupName(EQueryTickGroups::Default), false);
+			RegisterTickGroup(GetQueryTickGroupName(EQueryTickGroups::Update),
+				Phase, {}, GetQueryTickGroupName(EQueryTickGroups::PreUpdate), false);
+			RegisterTickGroup(GetQueryTickGroupName(EQueryTickGroups::PostUpdate),
+				Phase, {}, GetQueryTickGroupName(EQueryTickGroups::Update), false);
+
+			// After everything has processed sync the data in TEDS to external sources.
+			RegisterTickGroup(GetQueryTickGroupName(EQueryTickGroups::SyncDataStorageToExternal),
+				Phase, {}, GetQueryTickGroupName(EQueryTickGroups::PostUpdate), false);
+
+			// Update any widgets with data from TEDS.
 			RegisterTickGroup(GetQueryTickGroupName(EQueryTickGroups::SyncWidgets),
-				Phase, {}, GetQueryTickGroupName(EQueryTickGroups::SyncExternalToDataStorage), true /* Needs main thread*/);
+				Phase, {}, GetQueryTickGroupName(EQueryTickGroups::PostUpdate), true /* Needs main thread*/);
 		}
 	}
 }
@@ -230,7 +248,7 @@ void UTypedElementDatabase::OnPostMassTick(float DeltaTime)
 {
 	checkf(IsAvailable(), TEXT("Typed Element Database was ticked while it's not ready."));
 	
-	Environment->GetScratchBuffer().BatchDelete();
+	Environment->NextUpdateCycle();
 }
 
 TSharedPtr<FMassEntityManager> UTypedElementDatabase::GetActiveMutableEditorEntityManager()
@@ -755,18 +773,18 @@ bool UTypedElementDatabase::MatchesColumns(TypedElementDataStorage::RowHandle Ro
 void UTypedElementDatabase::RegisterTickGroup(
 	FName GroupName, EQueryTickPhase Phase, FName BeforeGroup, FName AfterGroup, bool bRequiresMainThread)
 {
-	Queries.RegisterTickGroup(GroupName, Phase, BeforeGroup, AfterGroup, bRequiresMainThread);
+	Environment->GetQueryStore().RegisterTickGroup(GroupName, Phase, BeforeGroup, AfterGroup, bRequiresMainThread);
 }
 
 void UTypedElementDatabase::UnregisterTickGroup(FName GroupName, EQueryTickPhase Phase)
 {
-	Queries.UnregisterTickGroup(GroupName, Phase);
+	Environment->GetQueryStore().UnregisterTickGroup(GroupName, Phase);
 }
 
 TypedElementQueryHandle UTypedElementDatabase::RegisterQuery(FQueryDescription&& Query)
 {
 	return (ActiveEditorEntityManager && ActiveEditorPhaseManager)
-		? Queries.RegisterQuery(MoveTemp(Query), *Environment, *ActiveEditorEntityManager, *ActiveEditorPhaseManager).Packed()
+		? Environment->GetQueryStore().RegisterQuery(MoveTemp(Query), *Environment, *ActiveEditorEntityManager, *ActiveEditorPhaseManager).Packed()
 		: TypedElementInvalidQueryHandle;
 }
 
@@ -775,14 +793,14 @@ void UTypedElementDatabase::UnregisterQuery(TypedElementQueryHandle Query)
 	if (ActiveEditorPhaseManager)
 	{
 		const FTypedElementExtendedQueryStore::Handle StorageHandle(Query);
-		Queries.UnregisterQuery(StorageHandle, *ActiveEditorPhaseManager);
+		Environment->GetQueryStore().UnregisterQuery(StorageHandle, *ActiveEditorPhaseManager);
 	}
 }
 
 const ITypedElementDataStorageInterface::FQueryDescription& UTypedElementDatabase::GetQueryDescription(TypedElementQueryHandle Query) const
 {
 	const FTypedElementExtendedQueryStore::Handle StorageHandle(Query);
-	return Queries.GetQueryDescription(StorageHandle);
+	return Environment->GetQueryStore().GetQueryDescription(StorageHandle);
 }
 
 FName UTypedElementDatabase::GetQueryTickGroupName(EQueryTickGroups Group) const
@@ -790,7 +808,13 @@ FName UTypedElementDatabase::GetQueryTickGroupName(EQueryTickGroups Group) const
 	switch (Group)
 	{
 		case EQueryTickGroups::Default:
-			return NAME_None;
+			return TickGroupName_Default;
+		case EQueryTickGroups::PreUpdate:
+			return TickGroupName_PreUpdate;
+		case EQueryTickGroups::Update:
+			return TickGroupName_Update;
+		case EQueryTickGroups::PostUpdate:
+			return TickGroupName_PostUpdate;
 		case EQueryTickGroups::SyncExternalToDataStorage:
 			return UE::Mass::ProcessorGroupNames::SyncWorldToMass;
 		case EQueryTickGroups::SyncDataStorageToExternal:
@@ -810,7 +834,7 @@ ITypedElementDataStorageInterface::FQueryResult UTypedElementDatabase::RunQuery(
 	if (ActiveEditorEntityManager)
 	{
 		const FTypedElementExtendedQueryStore::Handle StorageHandle(Query);
-		return Queries.RunQuery(*ActiveEditorEntityManager, StorageHandle);
+		return Environment->GetQueryStore().RunQuery(*ActiveEditorEntityManager, StorageHandle);
 	}
 	else
 	{
@@ -826,7 +850,7 @@ ITypedElementDataStorageInterface::FQueryResult UTypedElementDatabase::RunQuery(
 	if (ActiveEditorEntityManager)
 	{
 		const FTypedElementExtendedQueryStore::Handle StorageHandle(Query);
-		return Queries.RunQuery(*ActiveEditorEntityManager, *Environment, StorageHandle, Callback);
+		return Environment->GetQueryStore().RunQuery(*ActiveEditorEntityManager, *Environment, StorageHandle, Callback);
 	}
 	else
 	{
@@ -834,6 +858,14 @@ ITypedElementDataStorageInterface::FQueryResult UTypedElementDatabase::RunQuery(
 	}
 }
 
+void UTypedElementDatabase::ActivateQueries(FName ActivationName)
+{
+	if (ActiveEditorEntityManager)
+	{
+		Environment->GetQueryStore().ActivateQueries(ActivationName);
+	}
+}
+
 TypedElementDataStorage::RowHandle UTypedElementDatabase::FindIndexedRow(TypedElementDataStorage::IndexHash Index) const
 {
 	return Environment->GetIndexTable().FindIndexedRow(Index);
@@ -878,8 +910,7 @@ void UTypedElementDatabase::PreparePhase(EQueryTickPhase Phase, float DeltaTime)
 {
 	if (ActiveEditorEntityManager)
 	{
-		Environment->NextUpdateCycle();
-		Queries.RunPhasePreambleQueries(*ActiveEditorEntityManager, *Environment, Phase, DeltaTime);
+		Environment->GetQueryStore().RunPhasePreambleQueries(*ActiveEditorEntityManager, *Environment, Phase, DeltaTime);
 	}
 }
 
@@ -887,7 +918,7 @@ void UTypedElementDatabase::FinalizePhase(EQueryTickPhase Phase, float DeltaTime
 {
 	if (ActiveEditorEntityManager)
 	{
-		Queries.RunPhasePostambleQueries(*ActiveEditorEntityManager, *Environment, Phase, DeltaTime);
+		Environment->GetQueryStore().RunPhasePostambleQueries(*ActiveEditorEntityManager, *Environment, Phase, DeltaTime);
 	}
 }
 
@@ -903,7 +934,7 @@ void UTypedElementDatabase::Reset()
 
 	if (ActiveEditorPhaseManager)
 	{
-		Queries.Clear(*ActiveEditorPhaseManager.Get());
+		Environment->GetQueryStore().Clear(*ActiveEditorPhaseManager.Get());
 	}
 	Tables.Reset();
 	TableNameLookup.Reset();
@@ -914,7 +945,7 @@ void UTypedElementDatabase::Reset()
 
 void UTypedElementDatabase::DebugPrintQueryCallbacks(FOutputDevice& Output)
 {
-	Queries.DebugPrintQueryCallbacks(Output);
+	Environment->GetQueryStore().DebugPrintQueryCallbacks(Output);
 }
 
 void UTypedElementDatabase::AddReferencedObjects(UObject* InThis, FReferenceCollector& Collector)
diff --git a/Engine/Plugins/Experimental/TypedElementsDataStorage/Source/TypedElementsDataStorage/Private/TypedElementDatabase.h b/Engine/Plugins/Experimental/TypedElementsDataStorage/Source/TypedElementsDataStorage/Private/TypedElementDatabase.h
--- a/Engine/Plugins/Experimental/TypedElementsDataStorage/Source/TypedElementsDataStorage/Private/TypedElementDatabase.h
+++ b/Engine/Plugins/Experimental/TypedElementsDataStorage/Source/TypedElementsDataStorage/Private/TypedElementDatabase.h
@@ -6,7 +6,6 @@
 #include "Elements/Interfaces/TypedElementDataStorageInterface.h"
 #include "MassArchetypeTypes.h"
 #include "Misc/TVariant.h"
-#include "Queries/TypedElementExtendedQueryStore.h"
 #include "Templates/SharedPointer.h"
 #include "TypedElementDatabaseCommandBuffer.h"
 #include "TypedElementDatabaseEnvironment.h"
@@ -126,6 +125,7 @@ public:
 	FName GetQueryTickGroupName(EQueryTickGroups Group) const override;
 	FQueryResult RunQuery(TypedElementQueryHandle Query) override;
 	FQueryResult RunQuery(TypedElementQueryHandle Query, DirectQueryCallbackRef Callback) override;
+	void ActivateQueries(FName ActivationName) override;
 
 	TypedElementDataStorage::RowHandle FindIndexedRow(TypedElementDataStorage::IndexHash Index) const override;
 	void IndexRow(TypedElementDataStorage::IndexHash Index, TypedElementDataStorage::RowHandle Row) override;
@@ -156,6 +156,10 @@ private:
 		TObjectPtr<UTypedElementDataStorageFactory> Instance;
 	};
 	
+	static const FName TickGroupName_Default;
+	static const FName TickGroupName_PreUpdate;
+	static const FName TickGroupName_Update;
+	static const FName TickGroupName_PostUpdate;
 	static const FName TickGroupName_SyncWidget;
 	
 	TArray<FMassArchetypeHandle> Tables;
@@ -165,8 +169,7 @@ private:
 	TArray<FFactoryTypePair> Factories;
 
 	TUniquePtr<FTypedElementDatabaseEnvironment> Environment;
-	FTypedElementExtendedQueryStore Queries;
-
+	
 	FTypedElementDatabaseCommandBuffer::CommandBuffer DeferredCommands;
 	
 	FTypedElementOnDataStorageUpdate OnUpdateDelegate;
diff --git a/Engine/Plugins/Experimental/TypedElementsDataStorage/Source/TypedElementsDataStorage/Private/TypedElementDatabaseEnvironment.cpp b/Engine/Plugins/Experimental/TypedElementsDataStorage/Source/TypedElementsDataStorage/Private/TypedElementDatabaseEnvironment.cpp
--- a/Engine/Plugins/Experimental/TypedElementsDataStorage/Source/TypedElementsDataStorage/Private/TypedElementDatabaseEnvironment.cpp
+++ b/Engine/Plugins/Experimental/TypedElementsDataStorage/Source/TypedElementsDataStorage/Private/TypedElementDatabaseEnvironment.cpp
@@ -4,8 +4,7 @@
 
 FTypedElementDatabaseEnvironment::FTypedElementDatabaseEnvironment(
 	FMassEntityManager& InMassEntityManager, FMassProcessingPhaseManager& InMassPhaseManager)
-	: ScratchBuffer(MakeShared<FTypedElementDatabaseScratchBuffer>())
-	, MassEntityManager(InMassEntityManager)
+	: MassEntityManager(InMassEntityManager)
 	, MassPhaseManager(InMassPhaseManager)
 {
 }
@@ -22,12 +21,22 @@ const FTypedElementDatabaseIndexTable& FTypedElementDatabaseEnvironment::GetInde
 
 FTypedElementDatabaseScratchBuffer& FTypedElementDatabaseEnvironment::GetScratchBuffer()
 {
-	return *ScratchBuffer;
+	return ScratchBuffer;
 }
 
 const FTypedElementDatabaseScratchBuffer& FTypedElementDatabaseEnvironment::GetScratchBuffer() const
 {
-	return *ScratchBuffer;
+	return ScratchBuffer;
+}
+
+FTypedElementExtendedQueryStore& FTypedElementDatabaseEnvironment::GetQueryStore()
+{
+	return Queries;
+}
+
+const FTypedElementExtendedQueryStore& FTypedElementDatabaseEnvironment::GetQueryStore() const
+{
+	return Queries;
 }
 
 FMassEntityManager& FTypedElementDatabaseEnvironment::GetMassEntityManager()
@@ -52,6 +61,8 @@ const FMassProcessingPhaseManager& FTypedElementDatabaseEnvironment::GetMassPhas
 
 void FTypedElementDatabaseEnvironment::NextUpdateCycle()
 {
+	Queries.UpdateActivatableQueries();
+	ScratchBuffer.BatchDelete();
 	UpdateCycleId++;
 }
 
diff --git a/Engine/Plugins/Experimental/TypedElementsDataStorage/Source/TypedElementsDataStorage/Private/TypedElementDatabaseEnvironment.h b/Engine/Plugins/Experimental/TypedElementsDataStorage/Source/TypedElementsDataStorage/Private/TypedElementDatabaseEnvironment.h
--- a/Engine/Plugins/Experimental/TypedElementsDataStorage/Source/TypedElementsDataStorage/Private/TypedElementDatabaseEnvironment.h
+++ b/Engine/Plugins/Experimental/TypedElementsDataStorage/Source/TypedElementsDataStorage/Private/TypedElementDatabaseEnvironment.h
@@ -6,6 +6,7 @@
 #include "MassProcessingPhaseManager.h"
 #include "TypedElementDatabaseScratchBuffer.h"
 #include "TypedElementDatabaseIndexTable.h"
+#include "Queries/TypedElementExtendedQueryStore.h"
 
 class FTypedElementDatabaseEnvironment final
 {
@@ -18,6 +19,9 @@ public:
 	FTypedElementDatabaseScratchBuffer& GetScratchBuffer();
 	const FTypedElementDatabaseScratchBuffer& GetScratchBuffer() const;
 
+	FTypedElementExtendedQueryStore& GetQueryStore();
+	const FTypedElementExtendedQueryStore& GetQueryStore() const;
+
 	FMassEntityManager& GetMassEntityManager();
 	const FMassEntityManager& GetMassEntityManager() const;
 	
@@ -29,7 +33,8 @@ public:
 
 private:
 	FTypedElementDatabaseIndexTable IndexTable;
-	TSharedPtr<FTypedElementDatabaseScratchBuffer> ScratchBuffer;
+	FTypedElementDatabaseScratchBuffer ScratchBuffer;
+	FTypedElementExtendedQueryStore Queries;
 	
 	FMassEntityManager& MassEntityManager;
 	FMassProcessingPhaseManager& MassPhaseManager;
diff --git a/Engine/Source/Runtime/TypedElementFramework/Private/Elements/Framework/TypedElementQueryBuilder.cpp b/Engine/Source/Runtime/TypedElementFramework/Private/Elements/Framework/TypedElementQueryBuilder.cpp
--- a/Engine/Source/Runtime/TypedElementFramework/Private/Elements/Framework/TypedElementQueryBuilder.cpp
+++ b/Engine/Source/Runtime/TypedElementFramework/Private/Elements/Framework/TypedElementQueryBuilder.cpp
@@ -250,6 +250,12 @@ namespace TypedElementQueryBuilder
 		return *this;
 	}
 
+	FProcessor& FProcessor::MakeActivatable(FName Name)
+	{
+		ActivationName = Name;
+		return *this;
+	}
+
 
 	/**
 	 * FObserver
@@ -278,6 +284,12 @@ namespace TypedElementQueryBuilder
 		return *this;
 	}
 
+	FObserver& FObserver::MakeActivatable(FName Name)
+	{
+		ActivationName = Name;
+		return *this;
+	}
+
 
 	/**
 	 * FPhaseAmble
@@ -306,6 +318,12 @@ namespace TypedElementQueryBuilder
 		return *this;
 	}
 
+	FPhaseAmble& FPhaseAmble::MakeActivatable(FName Name)
+	{
+		ActivationName = Name;
+		return *this;
+	}
+
 
 	/**
 	 * Select
diff --git a/Engine/Source/Runtime/TypedElementFramework/Public/Elements/Common/TypedElementQueryDescription.h b/Engine/Source/Runtime/TypedElementFramework/Public/Elements/Common/TypedElementQueryDescription.h
--- a/Engine/Source/Runtime/TypedElementFramework/Public/Elements/Common/TypedElementQueryDescription.h
+++ b/Engine/Source/Runtime/TypedElementFramework/Public/Elements/Common/TypedElementQueryDescription.h
@@ -40,31 +40,18 @@ namespace TypedElementDataStorage
 			Max //< Value indicating the maximum value in this enum. Not to be used as an enum value.
 		};
 
-		using OperatorIndex = int32;
 		enum class EOperatorType : uint16
 		{
 			SimpleAll,			//< Unary: Type
 			SimpleAny,			//< Unary: Type
 			SimpleNone,			//< Unary: Type
 			SimpleOptional,		//< Unary: Type
-			And,				//< Binary: left operator index, right operator index
-			Or,					//< Binary: left operator index, right operator index
-			Not,				//< Unary: condition index
-			Type,				//< Unary: Type
-
+		
 			Max //< Value indicating the maximum value in this enum. Not to be used as an enum value.
 		};
 
-		struct FBinaryOperator final
-		{
-			OperatorIndex Left;
-			OperatorIndex Right;
-		};
-
 		union FOperator
 		{
-			FBinaryOperator Binary;
-			OperatorIndex Unary;
 			TWeakObjectPtr<const UScriptStruct> Type;
 		};
 
@@ -75,12 +62,18 @@ namespace TypedElementDataStorage
 			QueryCallback Function;
 			FName Name;
 			FName Group;
+			/** If a name is set, it indicates the query callback will not be run unless the ActivationCount is greater than zero. */
+			FName ActivationName;
 			const UScriptStruct* MonitoredType{ nullptr };
 			EQueryCallbackType Type{ EQueryCallbackType::None };
-			EQueryTickPhase Phase;
+			EQueryTickPhase Phase{ EQueryTickPhase::FrameEnd };
+			/**
+			 * The number of remaining iterations for a activatable query callback. If this is higher than 0, the query callback will be 
+			 * called. If ActivationName is set, this value will be decremented by one at the end of the update cycle.
+			 */
+			uint8 ActivationCount{ 255 };
 			bool bForceToGameThread{ false };
 		};
-
 		FCallbackData Callback;
 
 		// The list of arrays below are required to remain in the same order as they're added as the function binding expects certain entries
diff --git a/Engine/Source/Runtime/TypedElementFramework/Public/Elements/Common/TypedElementQueryTypes.h b/Engine/Source/Runtime/TypedElementFramework/Public/Elements/Common/TypedElementQueryTypes.h
--- a/Engine/Source/Runtime/TypedElementFramework/Public/Elements/Common/TypedElementQueryTypes.h
+++ b/Engine/Source/Runtime/TypedElementFramework/Public/Elements/Common/TypedElementQueryTypes.h
@@ -21,6 +21,15 @@ namespace TypedElementDataStorage
 		/** The standard group to run work in. */
 		Default,
 
+		/** General group that prepares for updating. */
+		PreUpdate,
+
+		/** General group that's used for updating data. */
+		Update,
+
+		/** General group that completes for updating. */
+		PostUpdate,
+
 		/**
 		 * The group for queries that need to sync data from external sources such as subsystems or the world into
 		 * the Data Storage. These typically run early in a phase.
diff --git a/Engine/Source/Runtime/TypedElementFramework/Public/Elements/Framework/TypedElementQueryBuilder.h b/Engine/Source/Runtime/TypedElementFramework/Public/Elements/Framework/TypedElementQueryBuilder.h
--- a/Engine/Source/Runtime/TypedElementFramework/Public/Elements/Framework/TypedElementQueryBuilder.h
+++ b/Engine/Source/Runtime/TypedElementFramework/Public/Elements/Framework/TypedElementQueryBuilder.h
@@ -174,11 +174,13 @@ namespace TypedElementQueryBuilder
 		TYPEDELEMENTFRAMEWORK_API FProcessor& SetBeforeGroup(FName GroupName);
 		TYPEDELEMENTFRAMEWORK_API FProcessor& SetAfterGroup(FName GroupName);
 		TYPEDELEMENTFRAMEWORK_API FProcessor& ForceToGameThread(bool bForce);
+		TYPEDELEMENTFRAMEWORK_API FProcessor& MakeActivatable(FName Name);
 		
 		ITypedElementDataStorageInterface::EQueryTickPhase Phase;
 		FName Group;
 		FName BeforeGroup;
 		FName AfterGroup;
+		FName ActivationName;
 		bool bForceToGameThread{ false };
 	};
 
@@ -202,9 +204,11 @@ namespace TypedElementQueryBuilder
 		template<typename ColumnType>
 		FObserver& SetMonitoredColumn();
 		TYPEDELEMENTFRAMEWORK_API FObserver& ForceToGameThread(bool bForce);
+		TYPEDELEMENTFRAMEWORK_API FObserver& MakeActivatable(FName Name);
 
 		const UScriptStruct* Monitor;
 		EEvent Event;
+		FName ActivationName;
 		bool bForceToGameThread{ false };
 	};
 
@@ -220,9 +224,11 @@ namespace TypedElementQueryBuilder
 		TYPEDELEMENTFRAMEWORK_API FPhaseAmble& SetLocation(ELocation NewLocation);
 		TYPEDELEMENTFRAMEWORK_API FPhaseAmble& SetPhase(ITypedElementDataStorageInterface::EQueryTickPhase NewPhase);
 		TYPEDELEMENTFRAMEWORK_API FPhaseAmble& ForceToGameThread(bool bForce);
+		TYPEDELEMENTFRAMEWORK_API FPhaseAmble& MakeActivatable(FName Name);
 
 		ITypedElementDataStorageInterface::EQueryTickPhase Phase;
 		ELocation Location;
+		FName ActivationName;
 		bool bForceToGameThread{ false };
 	};
 
diff --git a/Engine/Source/Runtime/TypedElementFramework/Public/Elements/Framework/TypedElementQueryBuilder.inl b/Engine/Source/Runtime/TypedElementFramework/Public/Elements/Framework/TypedElementQueryBuilder.inl
--- a/Engine/Source/Runtime/TypedElementFramework/Public/Elements/Framework/TypedElementQueryBuilder.inl
+++ b/Engine/Source/Runtime/TypedElementFramework/Public/Elements/Framework/TypedElementQueryBuilder.inl
@@ -1046,6 +1046,8 @@ namespace TypedElementQueryBuilder
 			Query.Callback.Type = TypedElementDataStorage::EQueryCallbackType::Processor;
 			Query.Callback.Phase = Processor.Phase;
 			Query.Callback.Group = Processor.Group;
+			Query.Callback.ActivationName = Processor.ActivationName;
+			Query.Callback.ActivationCount = Processor.ActivationName.IsNone() ? 255 : 0;
 			if (!Processor.BeforeGroup.IsNone())
 			{
 				Query.Callback.BeforeGroups.Add(Processor.BeforeGroup);
@@ -1068,6 +1070,8 @@ namespace TypedElementQueryBuilder
 				Query.Callback.Type = TypedElementDataStorage::EQueryCallbackType::ObserveRemove;
 				break;
 			}
+			Query.Callback.ActivationName = Observer.ActivationName;
+			Query.Callback.ActivationCount = Observer.ActivationName.IsNone() ? 255 : 0;
 			Query.Callback.MonitoredType = Observer.Monitor;
 			Query.Callback.bForceToGameThread = Observer.bForceToGameThread;
 		}
@@ -1083,6 +1087,8 @@ namespace TypedElementQueryBuilder
 				Query.Callback.Type = TypedElementDataStorage::EQueryCallbackType::PhaseFinalization;
 				break;
 			}
+			Query.Callback.ActivationName = PhaseAmble.ActivationName;
+			Query.Callback.ActivationCount = PhaseAmble.ActivationName.IsNone() ? 255 : 0;
 			Query.Callback.Phase = PhaseAmble.Phase;
 			Query.Callback.bForceToGameThread = PhaseAmble.bForceToGameThread;
 		}
diff --git a/Engine/Source/Runtime/TypedElementFramework/Public/Elements/Interfaces/TypedElementDataStorageInterface.h b/Engine/Source/Runtime/TypedElementFramework/Public/Elements/Interfaces/TypedElementDataStorageInterface.h
--- a/Engine/Source/Runtime/TypedElementFramework/Public/Elements/Interfaces/TypedElementDataStorageInterface.h
+++ b/Engine/Source/Runtime/TypedElementFramework/Public/Elements/Interfaces/TypedElementDataStorageInterface.h
@@ -274,6 +274,12 @@ public:
 	 * may be called multiple times. If the query handle is invalid or has been deleted nothing happens and the callback won't be called
 	 */
 	virtual FQueryResult RunQuery(TypedElementQueryHandle Query, DirectQueryCallbackRef Callback) = 0;
+	/**
+	 * Triggers all queries registered under the activation name to run for one update cycle. The activatable queries will be activated at
+	 * start of the cycle and disabled at the end of the cycle and act like regular queries for that cycle. This includes not running
+	 * if there are no columns to match against.
+	 */
+	virtual void ActivateQueries(FName ActivationName) = 0;
 	
 	/**
 	 * @section Indexing
diff --git a/Engine/Source/Runtime/TypedElementFramework/Public/Elements/Interfaces/TypedElementQueryStorageInterfaces.h b/Engine/Source/Runtime/TypedElementFramework/Public/Elements/Interfaces/TypedElementQueryStorageInterfaces.h
--- a/Engine/Source/Runtime/TypedElementFramework/Public/Elements/Interfaces/TypedElementQueryStorageInterfaces.h
+++ b/Engine/Source/Runtime/TypedElementFramework/Public/Elements/Interfaces/TypedElementQueryStorageInterfaces.h
@@ -101,6 +101,13 @@ namespace TypedElementDataStorage
 		 */
 		virtual uint64 GetUpdateCycleId() const = 0;
 
+		/**
+		 * Triggers all queries registered under the activation name to run for one update cycle. The activatable queries will be activated at
+		 * start of the cycle and disabled at the end of the cycle and act like regular queries for that cycle. This includes not running
+		 * if there are not columns to match against.
+		 */
+		virtual void ActivateQueries(FName ActivationName) = 0;
+
 		/**
 		 * Adds the provided column to the requested row.
 		 *
