TEDS - Make sure the EntityManager is available before trying to unregister queries

diff --git a/Engine/Plugins/Experimental/TypedElementsDataStorage/Source/TypedElementsDataStorage/Private/Queries/TypedElementExtendedQueryStore.cpp b/Engine/Plugins/Experimental/TypedElementsDataStorage/Source/TypedElementsDataStorage/Private/Queries/TypedElementExtendedQueryStore.cpp
--- a/Engine/Plugins/Experimental/TypedElementsDataStorage/Source/TypedElementsDataStorage/Private/Queries/TypedElementExtendedQueryStore.cpp
+++ b/Engine/Plugins/Experimental/TypedElementsDataStorage/Source/TypedElementsDataStorage/Private/Queries/TypedElementExtendedQueryStore.cpp
@@ -36,27 +36,27 @@ FTypedElementExtendedQueryStore::Handle FTypedElementExtendedQueryStore::Registe
 	return Result;
 }
 
-void FTypedElementExtendedQueryStore::UnregisterQuery(Handle Query, FMassProcessingPhaseManager& PhaseManager)
+void FTypedElementExtendedQueryStore::UnregisterQuery(Handle Query, FMassEntityManager& EntityManager, FMassProcessingPhaseManager& PhaseManager)
 {
 	if (FTypedElementExtendedQuery* QueryData = Get(Query))
 	{
-		UnregisterQueryData(Query, *QueryData, PhaseManager);
+		UnregisterQueryData(Query, *QueryData, EntityManager, PhaseManager);
 		Queries.Remove(Query);
 	}
 }
 
-void FTypedElementExtendedQueryStore::Clear(FMassProcessingPhaseManager& PhaseManager)
+void FTypedElementExtendedQueryStore::Clear(FMassEntityManager& EntityManager, FMassProcessingPhaseManager& PhaseManager)
 {
 	TickGroupDescriptions.Empty();
 
-	Queries.ListAliveEntries([this, &PhaseManager](Handle Query, FTypedElementExtendedQuery& QueryData)
+	Queries.ListAliveEntries([this, &EntityManager, &PhaseManager](Handle Query, FTypedElementExtendedQuery& QueryData)
 		{
 			if (QueryData.Processor && QueryData.Processor->IsA<UTypedElementQueryObserverCallbackAdapterProcessorBase>())
 			{
 				// Observers can't be unregistered at this point, so skip these for now.
 				return;
 			}
-			UnregisterQueryData(Query, QueryData, PhaseManager);
+			UnregisterQueryData(Query, QueryData, EntityManager, PhaseManager);
 		});
 }
 
@@ -747,7 +747,7 @@ void FTypedElementExtendedQueryStore::RunPhasePreOrPostAmbleQueries(FMassEntityM
 	}
 }
 
-void FTypedElementExtendedQueryStore::UnregisterQueryData(Handle Query, FTypedElementExtendedQuery& QueryData, FMassProcessingPhaseManager& PhaseManager)
+void FTypedElementExtendedQueryStore::UnregisterQueryData(Handle Query, FTypedElementExtendedQuery& QueryData, FMassEntityManager& EntityManager, FMassProcessingPhaseManager& PhaseManager)
 {
 	if (!QueryData.Description.Callback.ActivationName.IsNone())
 	{
@@ -766,9 +766,9 @@ void FTypedElementExtendedQueryStore::UnregisterQueryData(Handle Query, FTypedEl
 		{
 			UTypedElementQueryObserverCallbackAdapterProcessorBase* Observer =
 				static_cast<UTypedElementQueryObserverCallbackAdapterProcessorBase*>(QueryData.Processor.Get());
-			if (ensure(Observer && PhaseManager.GetEntityManager().IsValid()))
+			if (ensure(Observer))
 			{
-				PhaseManager.GetEntityManager()->GetObserverManager().RemoveObserverInstance(*Observer->GetObservedType(), Observer->GetObservedOperation(), *Observer);
+				EntityManager.GetObserverManager().RemoveObserverInstance(*Observer->GetObservedType(), Observer->GetObservedOperation(), *Observer);
 			}
 		}
 		else
diff --git a/Engine/Plugins/Experimental/TypedElementsDataStorage/Source/TypedElementsDataStorage/Private/Queries/TypedElementExtendedQueryStore.h b/Engine/Plugins/Experimental/TypedElementsDataStorage/Source/TypedElementsDataStorage/Private/Queries/TypedElementExtendedQueryStore.h
--- a/Engine/Plugins/Experimental/TypedElementsDataStorage/Source/TypedElementsDataStorage/Private/Queries/TypedElementExtendedQueryStore.h
+++ b/Engine/Plugins/Experimental/TypedElementsDataStorage/Source/TypedElementsDataStorage/Private/Queries/TypedElementExtendedQueryStore.h
@@ -44,10 +44,10 @@ public:
 		FMassEntityManager& EntityManager,
 		FMassProcessingPhaseManager& PhaseManager);
 	/** Removes the query at the given handle if still alive and otherwise does nothing. */
-	void UnregisterQuery(Handle Query, FMassProcessingPhaseManager& PhaseManager);
+	void UnregisterQuery(Handle Query, FMassEntityManager& EntityManager, FMassProcessingPhaseManager& PhaseManager);
 
 	/** Removes all data in the query store. */
-	void Clear(FMassProcessingPhaseManager& PhaseManager);
+	void Clear(FMassEntityManager& EntityManager, FMassProcessingPhaseManager& PhaseManager);
 	
 	/** Register the defaults for a tick group. These will be applied on top of any settings provided with a query registration. */
 	void RegisterTickGroup(FName GroupName, ITypedElementDataStorageInterface::EQueryTickPhase Phase,
@@ -177,7 +177,7 @@ private:
 	void RunPhasePreOrPostAmbleQueries(FMassEntityManager& EntityManager, FTypedElementDatabaseEnvironment& Environment,
 		ITypedElementDataStorageInterface::EQueryTickPhase Phase, float DeltaTime, TArray<Handle>& QueryHandles);
 
-	void UnregisterQueryData(Handle Query, FTypedElementExtendedQuery& QueryData, FMassProcessingPhaseManager& PhaseManager);
+	void UnregisterQueryData(Handle Query, FTypedElementExtendedQuery& QueryData, FMassEntityManager& EntityManager, FMassProcessingPhaseManager& PhaseManager);
 
 	static const ITypedElementDataStorageInterface::FQueryDescription EmptyDescription;
 
diff --git a/Engine/Plugins/Experimental/TypedElementsDataStorage/Source/TypedElementsDataStorage/Private/TypedElementDatabase.cpp b/Engine/Plugins/Experimental/TypedElementsDataStorage/Source/TypedElementsDataStorage/Private/TypedElementDatabase.cpp
--- a/Engine/Plugins/Experimental/TypedElementsDataStorage/Source/TypedElementsDataStorage/Private/TypedElementDatabase.cpp
+++ b/Engine/Plugins/Experimental/TypedElementsDataStorage/Source/TypedElementsDataStorage/Private/TypedElementDatabase.cpp
@@ -790,10 +790,10 @@ TypedElementQueryHandle UTypedElementDatabase::RegisterQuery(FQueryDescription&&
 
 void UTypedElementDatabase::UnregisterQuery(TypedElementQueryHandle Query)
 {
-	if (ActiveEditorPhaseManager)
+	if (ActiveEditorEntityManager && ActiveEditorPhaseManager)
 	{
 		const FTypedElementExtendedQueryStore::Handle StorageHandle(Query);
-		Environment->GetQueryStore().UnregisterQuery(StorageHandle, *ActiveEditorPhaseManager);
+		Environment->GetQueryStore().UnregisterQuery(StorageHandle, *ActiveEditorEntityManager, *ActiveEditorPhaseManager);
 	}
 }
 
@@ -932,9 +932,9 @@ void UTypedElementDatabase::Reset()
 	OnPostMassTickHandle.Reset();
 	OnPreMassTickHandle.Reset();
 
-	if (ActiveEditorPhaseManager)
+	if (ActiveEditorEntityManager && ActiveEditorPhaseManager)
 	{
-		Environment->GetQueryStore().Clear(*ActiveEditorPhaseManager.Get());
+		Environment->GetQueryStore().Clear(*ActiveEditorEntityManager.Get(), *ActiveEditorPhaseManager.Get());
 	}
 	Tables.Reset();
 	TableNameLookup.Reset();
