Replaced alert updates with a simpler sweeping approach.

diff --git a/Engine/Plugins/Experimental/TypedElementsDataStorage/Source/TypedElementsDataStorage/Private/Queries/TypedElementAlertQueries.cpp b/Engine/Plugins/Experimental/TypedElementsDataStorage/Source/TypedElementsDataStorage/Private/Queries/TypedElementAlertQueries.cpp
--- a/Engine/Plugins/Experimental/TypedElementsDataStorage/Source/TypedElementsDataStorage/Private/Queries/TypedElementAlertQueries.cpp
+++ b/Engine/Plugins/Experimental/TypedElementsDataStorage/Source/TypedElementsDataStorage/Private/Queries/TypedElementAlertQueries.cpp
@@ -146,15 +146,15 @@ FAutoConsoleCommand ClearSelectedAlertsConsoleCommand(
 		}
 ));
 
+const FName UTypedElementAlertQueriesFactory::AlertConditionName = FName(TEXT("Alerts"));
+
 void UTypedElementAlertQueriesFactory::RegisterQueries(ITypedElementDataStorageInterface& DataStorage)
 {
 	RegisterSubQueries(DataStorage);
-	
-	RegisterUpdateAlertsQueries(DataStorage);
-	
+	RegisterParentUpdatesQueries(DataStorage);
+	RegisterChildAlertUpdatesQueries(DataStorage);
 	RegisterOnAddQueries(DataStorage);
-	RegisterOnAlertRemoveQueries(DataStorage);
-	RegisterOnParentRemoveQueries(DataStorage);
+	RegisterOnRemoveQueries(DataStorage);
 }
 
 void UTypedElementAlertQueriesFactory::RegisterSubQueries(ITypedElementDataStorageInterface& DataStorage)
@@ -172,128 +172,103 @@ void UTypedElementAlertQueriesFactory::RegisterSubQueries(ITypedElementDataStora
 		.Compile());
 }
 
-void UTypedElementAlertQueriesFactory::RegisterUpdateAlertsQueries(ITypedElementDataStorageInterface& DataStorage)
+void UTypedElementAlertQueriesFactory::RegisterParentUpdatesQueries(ITypedElementDataStorageInterface& DataStorage)
 {
 	using namespace TypedElementDataStorage;
 	using namespace TypedElementQueryBuilder;
 
 	DataStorage.RegisterQuery(
 		Select(
-			TEXT("Update alerts"),
-			FProcessor(EQueryTickPhase::DuringPhysics, DataStorage.GetQueryTickGroupName(EQueryTickGroups::Default)),
-			[](IQueryContext& Context, RowHandle Row, FTypedElementAlertColumn& Alert, const FTypedElementParentColumn& Parent)
+			TEXT("Trigger alert update if alert's parent changed"),
+			FProcessor(EQueryTickPhase::FrameEnd, DataStorage.GetQueryTickGroupName(EQueryTickGroups::Default)),
+			[](IQueryContext& Context, FTypedElementAlertColumn& Alert, const FTypedElementParentColumn& Parent)
 			{
-				Update<true, false>(Context, Row, &Alert, nullptr, Parent, 0, 1);
+				if (Alert.CachedParent != Parent.Parent)
+				{
+					Alert.CachedParent = Parent.Parent;
+					Context.ActivateQueries(AlertConditionName);
+				}
 			})
 		.Where()
-			.Any<FTypedElementAlertUpdateTag, FTypedElementSyncFromWorldTag, FTypedElementSyncBackToWorldTag>()
-			.None<FTypedElementChildAlertColumn>()
-		.DependsOn()
-			.SubQuery(ChildAlertColumnReadWriteQuery)
-			.SubQuery(ParentReadOnlyQuery)
+			.Any<FTypedElementSyncBackToWorldTag, FTypedElementSyncFromWorldTag>()
 		.Compile());
 
 	DataStorage.RegisterQuery(
 		Select(
-			TEXT("Update child alerts"),
-			FProcessor(EQueryTickPhase::DuringPhysics, DataStorage.GetQueryTickGroupName(EQueryTickGroups::Default)),
-			[](IQueryContext& Context, RowHandle Row, FTypedElementChildAlertColumn& ChildAlert, const FTypedElementParentColumn& Parent)
+			TEXT("Trigger alert update if child alert's parent changed"),
+			FProcessor(EQueryTickPhase::FrameEnd, DataStorage.GetQueryTickGroupName(EQueryTickGroups::Default)),
+			[](IQueryContext& Context, FTypedElementChildAlertColumn& ChildAlert, const FTypedElementParentColumn& Parent)
 			{
-				Update<false, true>(Context, Row, nullptr, &ChildAlert, Parent, 0, 1);
+				if (ChildAlert.CachedParent != Parent.Parent)
+				{
+					ChildAlert.CachedParent = Parent.Parent;
+					Context.ActivateQueries(AlertConditionName);
+				}
 			})
 		.Where()
-			.Any<FTypedElementAlertUpdateTag, FTypedElementSyncFromWorldTag, FTypedElementSyncBackToWorldTag>()
-			.None<FTypedElementAlertColumn>()
-		.DependsOn()
-			.SubQuery(ChildAlertColumnReadWriteQuery)
-			.SubQuery(ParentReadOnlyQuery)
+			.Any<FTypedElementSyncBackToWorldTag, FTypedElementSyncFromWorldTag>()
 		.Compile());
+}
+
+void UTypedElementAlertQueriesFactory::RegisterChildAlertUpdatesQueries(ITypedElementDataStorageInterface& DataStorage)
+{
+	using namespace TypedElementDataStorage;
+	using namespace TypedElementQueryBuilder;
 
 	DataStorage.RegisterQuery(
 		Select(
-			TEXT("Update alerts and child alerts"),
-			FProcessor(EQueryTickPhase::DuringPhysics, DataStorage.GetQueryTickGroupName(EQueryTickGroups::Default)),
-			[](IQueryContext& Context, RowHandle Row, FTypedElementAlertColumn& Alert, FTypedElementChildAlertColumn& ChildAlert, 
-				const FTypedElementParentColumn& Parent)
+			TEXT("Add missing child alerts"),
+			FPhaseAmble(FPhaseAmble::ELocation::Preamble, EQueryTickPhase::FrameEnd)
+				.MakeActivatable(AlertConditionName),
+			[](IQueryContext& Context, RowHandle Row, FTypedElementAlertColumn& Alert, const FTypedElementParentColumn& Parent)
 			{
-				Update<true, true>(Context, Row, &Alert, &ChildAlert, Parent, 0, 1);
+				AddChildAlertsToHierarchy(Context, Parent.Parent, 0);
 			})
-		.Where()
-			.Any<FTypedElementAlertUpdateTag, FTypedElementSyncFromWorldTag, FTypedElementSyncBackToWorldTag>()
 		.DependsOn()
-			.SubQuery(ChildAlertColumnReadWriteQuery)
 			.SubQuery(ParentReadOnlyQuery)
 		.Compile());
-}
-
-template<bool bHasAlert, bool bHasChildAlert>
-void UTypedElementAlertQueriesFactory::Update(TypedElementDataStorage::IQueryContext& Context, TypedElementDataStorage::RowHandle Row,
-	FTypedElementAlertColumn* Alert, FTypedElementChildAlertColumn* ChildAlert, const FTypedElementParentColumn& Parent,
-	int32 ChildAlertColumnReadWriteSubquery, int32 ParentReadOnlySubquery)
-{
-	using namespace TypedElementDataStorage;
-	using namespace TypedElementQueryBuilder;
-
-	bool bAlertNeedsUpdating = false;
-	bool bChildAlertsNeedUpdating = false;
-	NeedsUpdating<bHasAlert, bHasChildAlert>(Alert, ChildAlert, Parent, bAlertNeedsUpdating, bChildAlertsNeedUpdating);
-
-	bool bAlertNeedsDecrementing = false;
-	bool bChildAlertsNeedDecrementing = false;
-	NeedsDecrementing<bHasAlert, bHasChildAlert>(Alert, ChildAlert, Parent, bAlertNeedsDecrementing, bChildAlertsNeedDecrementing);
 
-	if (bAlertNeedsUpdating || bChildAlertsNeedUpdating || bAlertNeedsDecrementing || bChildAlertsNeedDecrementing)
-	{
-		RowHandle CachedParentRow = GetParent<bHasAlert, bHasChildAlert>(Alert, ChildAlert);
-
-		// If needed, first decrement the parents.
-		if ((bAlertNeedsDecrementing || bChildAlertsNeedDecrementing) && Context.IsRowAvailable(CachedParentRow))
-		{
-			uint16 TotalErrorCounts = 0;
-			uint16 TotalWarningCounts = 0;
-			GetTotalCounts<bHasAlert, bHasChildAlert>(Alert, ChildAlert, TotalErrorCounts, TotalWarningCounts,
-				bAlertNeedsDecrementing, bChildAlertsNeedDecrementing);
-
-			if (TotalErrorCounts != 0 || TotalWarningCounts != 0)
+	DataStorage.RegisterQuery(
+		Select(
+			TEXT("Clear child alerts"),
+			FProcessor(EQueryTickPhase::FrameEnd, DataStorage.GetQueryTickGroupName(EQueryTickGroups::PreUpdate))
+				.MakeActivatable(AlertConditionName),
+			[](IQueryContext& Context, RowHandle Row, FTypedElementChildAlertColumn& ChildAlert)
 			{
-				DecrementParents(
-					Context, CachedParentRow, TotalErrorCounts, TotalWarningCounts, ChildAlertColumnReadWriteSubquery);
-				Context.AddColumns<FTypedElementSyncBackToWorldTag>(Row);
-
-				// Prevent a second delete if there's a re-entry.
-				MarkDecremented<bHasAlert, bHasChildAlert>(Alert, ChildAlert);
-				
-				// Columns that are needed may have been deleted, so wait a cycle to give them a change to process and be re-added.
-				GuaranteeUpdateReentry(Context, Row);
-				return;
-			}
-		}
+				ResetChildAlertCounters(ChildAlert);
+			})
+		.Compile());
 
-		// Increment all the parents.
-		RowHandle ParentRow = Parent.Parent;
-		if (Context.IsRowAvailable(ParentRow))
-		{
-			if (Context.HasColumn<FTypedElementChildAlertColumn>(ParentRow))
+	DataStorage.RegisterQuery(
+		Select(
+			TEXT("Increment child alerts"),
+			FProcessor(EQueryTickPhase::FrameEnd, DataStorage.GetQueryTickGroupName(EQueryTickGroups::Update))
+				.MakeActivatable(AlertConditionName),
+			[](IQueryContext& Context, RowHandle Row, FTypedElementAlertColumn& Alert)
 			{
-				// At this point all parents should have had a child alert column. It can be safely assumed that if a direct parent has
-				// child alert column, that this is recursively true for its parents.
-				uint16 TotalErrorCounts = 0;
-				uint16 TotalWarningCounts = 0;
-				GetTotalCounts<bHasAlert, bHasChildAlert>(Alert, ChildAlert, TotalErrorCounts, TotalWarningCounts, 
-					bAlertNeedsUpdating, bChildAlertsNeedUpdating);
-				IncrementParents(Context, ParentRow, TotalErrorCounts, TotalWarningCounts, ChildAlertColumnReadWriteSubquery);
-
-				UpdateParent<bHasAlert, bHasChildAlert>(Alert, ChildAlert, ParentRow);
-				Context.RemoveColumns<FTypedElementAlertUpdateTag>(Row);
-			}
-			else
+				IncrementParents(Context, Alert.CachedParent, Alert.AlertType, 0);
+			})
+		.DependsOn()
+			.SubQuery(ChildAlertColumnReadWriteQuery)
+		.Compile());
+
+	DataStorage.RegisterQuery(
+		Select(
+			TEXT("Remove unused child alerts"),
+			FProcessor(EQueryTickPhase::FrameEnd, DataStorage.GetQueryTickGroupName(EQueryTickGroups::PostUpdate))
+				.MakeActivatable(AlertConditionName),
+			[](IQueryContext& Context, RowHandle Row, FTypedElementChildAlertColumn& ChildAlert)
 			{
-				// The parent doesn't have a child alert, so add it and continue in the next cycle so the column is added.
-				AddChildAlertsToHierarchy(Context, ParentRow, ParentReadOnlySubquery);
-				GuaranteeUpdateReentry(Context, Row);
-			}
-		}
-	}
+				for (size_t It = 0; It < static_cast<size_t>(FTypedElementAlertColumnType::MAX); ++It)
+				{
+					if (ChildAlert.Counts[It] != 0)
+					{
+						return;
+					}
+				}
+				Context.RemoveColumns<FTypedElementChildAlertColumn>(Row);
+			})
+		.Compile());
 }
 
 void UTypedElementAlertQueriesFactory::RegisterOnAddQueries(ITypedElementDataStorageInterface& DataStorage)
@@ -307,7 +282,7 @@ void UTypedElementAlertQueriesFactory::RegisterOnAddQueries(ITypedElementDataSto
 			FObserver::OnAdd<FTypedElementAlertColumn>(),
 			[](IQueryContext& Context, RowHandle Row)
 			{
-				Context.AddColumns<FTypedElementAlertUpdateTag>(Row);
+				Context.ActivateQueries(AlertConditionName);
 			})
 		.Where()
 			.All<FTypedElementParentColumn>() // Only need to do an update pass if there are parents.
@@ -319,14 +294,14 @@ void UTypedElementAlertQueriesFactory::RegisterOnAddQueries(ITypedElementDataSto
 			FObserver::OnAdd<FTypedElementParentColumn>(),
 			[](IQueryContext& Context, RowHandle Row)
 			{
-				Context.AddColumns<FTypedElementAlertUpdateTag>(Row);
+				Context.ActivateQueries(AlertConditionName);
 			})
 		.Where()
 			.Any<FTypedElementAlertColumn, FTypedElementChildAlertColumn>()
 		.Compile());
 }
 
-void UTypedElementAlertQueriesFactory::RegisterOnAlertRemoveQueries(ITypedElementDataStorageInterface& DataStorage)
+void UTypedElementAlertQueriesFactory::RegisterOnRemoveQueries(ITypedElementDataStorageInterface& DataStorage)
 {
 	using namespace TypedElementDataStorage;
 	using namespace TypedElementQueryBuilder;
@@ -337,351 +312,72 @@ void UTypedElementAlertQueriesFactory::RegisterOnAlertRemoveQueries(ITypedElemen
 			FObserver::OnRemove<FTypedElementAlertColumn>(),
 			[](IQueryContext& Context, TypedElementDataStorage::RowHandle Row, FTypedElementAlertColumn& Alert)
 			{
-				if (Alert.RemoveCycleId != Context.GetUpdateCycleId())
-				{
-					DecrementParents(Context, Alert.CachedParent,
-						Alert.AlertType == FTypedElementAlertColumnType::Error ? 1 : 0,
-						Alert.AlertType == FTypedElementAlertColumnType::Warning ? 1 : 0,
-						0);
-					Alert.CachedParent = InvalidRowHandle;
-					Alert.RemoveCycleId = Context.GetUpdateCycleId();
-					Context.AddColumns<FTypedElementSyncBackToWorldTag>(Row);
-				}
-				Context.RemoveColumns<FTypedElementAlertActionColumn>(Row);
+				Context.ActivateQueries(AlertConditionName);
 			})
-		.DependsOn()
-			.SubQuery(ChildAlertColumnReadWriteQuery)
-		.Compile());
-
-	DataStorage.RegisterQuery(
-		Select(
-			TEXT("Remove child alert"),
-			FObserver::OnRemove<FTypedElementChildAlertColumn>(),
-			[](IQueryContext& Context, TypedElementDataStorage::RowHandle Row, FTypedElementChildAlertColumn& ChildAlert)
-			{
-				if (Context.GetUpdateCycleId() != ChildAlert.RemoveCycleId)
-				{
-					DecrementParents(Context, ChildAlert.CachedParent,
-						ChildAlert.Counts[static_cast<size_t>(FTypedElementAlertColumnType::Error)],
-						ChildAlert.Counts[static_cast<size_t>(FTypedElementAlertColumnType::Warning)],
-						0);
-					ChildAlert.RemoveCycleId = Context.GetUpdateCycleId();
-					Context.AddColumns<FTypedElementSyncBackToWorldTag>(Row);
-				}
-			})
-		.DependsOn()
-			.SubQuery(ChildAlertColumnReadWriteQuery)
 		.Compile());
-}
-
-void UTypedElementAlertQueriesFactory::RegisterOnParentRemoveQueries(ITypedElementDataStorageInterface& DataStorage)
-{
-	using namespace TypedElementDataStorage;
-	using namespace TypedElementQueryBuilder;
 
 	DataStorage.RegisterQuery(
 		Select(
 			TEXT("Update alert upon parent removal"),
 			FObserver::OnRemove<FTypedElementParentColumn>(),
-			[](IQueryContext& Context, TypedElementDataStorage::RowHandle Row, FTypedElementAlertColumn& Alert)
+			[](IQueryContext& Context, TypedElementDataStorage::RowHandle Row)
 			{
-				if (Alert.RemoveCycleId != Context.GetUpdateCycleId())
-				{
-					DecrementParents(Context, Alert.CachedParent,
-						Alert.AlertType == FTypedElementAlertColumnType::Error ? 1 : 0,
-						Alert.AlertType == FTypedElementAlertColumnType::Warning ? 1 : 0,
-						0);
-					Alert.CachedParent = InvalidRowHandle;
-					Alert.RemoveCycleId = Context.GetUpdateCycleId();
-					Context.AddColumns<FTypedElementSyncBackToWorldTag>(Row);
-				}
+				Context.ActivateQueries(AlertConditionName);
 			})
 		.Where()
-			.None<FTypedElementChildAlertColumn>()
-		.DependsOn()
-			.SubQuery(ChildAlertColumnReadWriteQuery)
-		.Compile());
-
-	DataStorage.RegisterQuery(
-		Select(
-			TEXT("Update child alert upon parent removal"),
-			FObserver::OnRemove<FTypedElementParentColumn>(),
-			[](IQueryContext& Context, TypedElementDataStorage::RowHandle Row, FTypedElementChildAlertColumn& ChildAlert)
-			{
-				if (Context.GetUpdateCycleId() != ChildAlert.RemoveCycleId)
-				{
-					DecrementParents(Context, ChildAlert.CachedParent,
-						ChildAlert.Counts[static_cast<size_t>(FTypedElementAlertColumnType::Error)],
-						ChildAlert.Counts[static_cast<size_t>(FTypedElementAlertColumnType::Warning)],
-						0);
-					ChildAlert.CachedParent = InvalidRowHandle;
-					ChildAlert.RemoveCycleId = Context.GetUpdateCycleId();
-					Context.AddColumns<FTypedElementSyncBackToWorldTag>(Row);
-				}
-			})
-		.Where()
-			.None<FTypedElementAlertColumn>()
-		.DependsOn()
-			.SubQuery(ChildAlertColumnReadWriteQuery)
-		.Compile());
-
-	DataStorage.RegisterQuery(
-		Select(
-			TEXT("Update alert and child alert upon parent removal"),
-			FObserver::OnRemove<FTypedElementParentColumn>(),
-			[](IQueryContext& Context, TypedElementDataStorage::RowHandle Row, FTypedElementAlertColumn& Alert, FTypedElementChildAlertColumn& ChildAlert)
-			{
-				uint64 UpdateCycleId = Context.GetUpdateCycleId();
-				uint16 DecrementCounts[static_cast<size_t>(FTypedElementAlertColumnType::MAX)]{};
-				bool bUpdate = false;
-
-				if (Alert.RemoveCycleId != UpdateCycleId)
-				{
-					DecrementCounts[static_cast<size_t>(Alert.AlertType)]++;
-					bUpdate = true;
-				}
-				if (Context.GetUpdateCycleId() != ChildAlert.RemoveCycleId)
-				{
-					for (size_t Counter = 0; Counter < static_cast<size_t>(FTypedElementAlertColumnType::MAX); ++Counter)
-					{
-						DecrementCounts[Counter] += ChildAlert.Counts[Counter];
-					}
-					bUpdate = true;
-				}
-
-				if (bUpdate)
-				{
-					DecrementParents(Context, ChildAlert.CachedParent,
-						DecrementCounts[static_cast<size_t>(FTypedElementAlertColumnType::Error)],
-						DecrementCounts[static_cast<size_t>(FTypedElementAlertColumnType::Warning)],
-						0);
-					Alert.CachedParent = InvalidRowHandle;
-					Alert.RemoveCycleId = UpdateCycleId;
-					ChildAlert.CachedParent = InvalidRowHandle;
-					ChildAlert.RemoveCycleId = UpdateCycleId;
-				}
-			})
-		.DependsOn()
-			.SubQuery(ChildAlertColumnReadWriteQuery)
+			.Any<FTypedElementAlertColumn, FTypedElementChildAlertColumn>()
 		.Compile());
 }
 
-template<bool bHasAlert, bool bHasChildAlert>
-void UTypedElementAlertQueriesFactory::NeedsUpdating(
-	FTypedElementAlertColumn* Alert, FTypedElementChildAlertColumn* ChildAlert, const FTypedElementParentColumn& Parent,
-	bool& bAlertNeedsUpdate, bool& bChildAlertsNeedUpdate)
-{
-	using namespace TypedElementDataStorage;
-
-	if constexpr (bHasAlert)
-	{
-		bAlertNeedsUpdate = Alert->CachedParent != Parent.Parent;
-	}
-
-	if constexpr (bHasChildAlert)
-	{
-		bChildAlertsNeedUpdate = ChildAlert->CachedParent != Parent.Parent;
-	}
-}
-
-template<bool bHasAlert, bool bHasChildAlert>
-void UTypedElementAlertQueriesFactory::NeedsDecrementing(FTypedElementAlertColumn* Alert, FTypedElementChildAlertColumn* ChildAlert,
-	const FTypedElementParentColumn& Parent, bool& bAlertNeedsDecrementing, bool& bChildAlertsNeedDecrementing)
-{
-	using namespace TypedElementDataStorage;
-
-	if constexpr (bHasAlert)
-	{
-		bAlertNeedsDecrementing = 
-			Alert->CachedParent != 0 && 
-			Alert->CachedParent != InvalidRowHandle && 
-			Alert->CachedParent != Parent.Parent;
-	}
-
-	if constexpr (bHasChildAlert)
-	{
-		bChildAlertsNeedDecrementing =
-			!ChildAlert->bHasDecremented &&
-			ChildAlert->CachedParent != Parent.Parent;
-	}
-}
-
-template<bool bHasAlert, bool bHasChildAlert>
-TypedElementDataStorage::RowHandle UTypedElementAlertQueriesFactory::GetParent(
-	FTypedElementAlertColumn* Alert, FTypedElementChildAlertColumn* ChildAlert)
-{
-	using namespace TypedElementDataStorage;
-
-	RowHandle Result = InvalidRowHandle;
-	if constexpr (bHasAlert)
-	{
-		Result = Alert->CachedParent;
-	}
-	if constexpr (bHasChildAlert)
-	{
-		if constexpr (bHasAlert)
-		{
-			checkf(Result == 0 || Result == InvalidRowHandle || Result == ChildAlert->CachedParent,
-				TEXT("Alert and Child Alert are not pointing at the same parent."));
-
-		}
-		Result = ChildAlert->CachedParent;
-	}
-	return Result;
-}
-
-template<bool bHasAlert, bool bHasChildAlert>
-void UTypedElementAlertQueriesFactory::GetTotalCounts(FTypedElementAlertColumn* Alert, FTypedElementChildAlertColumn* ChildAlert,
-	uint16& ErrorCount, uint16& WarningCount, bool IncludeAlert, bool IncludeChildAlerts)
-{
-	if constexpr (bHasAlert)
-	{
-		switch (Alert->AlertType)
-		{
-		case FTypedElementAlertColumnType::Error:
-			ErrorCount += IncludeAlert ? 1 : 0;
-			break;
-		case FTypedElementAlertColumnType::Warning:
-			WarningCount += IncludeAlert ? 1 : 0;
-			break;
-		default:
-			checkf(false, TEXT("Unexpected alert type %i"), static_cast<int>(Alert->AlertType));
-		}
-	}
-	if constexpr (bHasChildAlert)
-	{
-		ErrorCount += IncludeChildAlerts ? ChildAlert->Counts[static_cast<size_t>(FTypedElementAlertColumnType::Error)] : 0;
-		WarningCount += IncludeChildAlerts ? ChildAlert->Counts[static_cast<size_t>(FTypedElementAlertColumnType::Warning)] : 0;
-	}
-}
-
-template<bool bHasAlert, bool bHasChildAlert>
-void UTypedElementAlertQueriesFactory::MarkDecremented(FTypedElementAlertColumn* Alert, FTypedElementChildAlertColumn* ChildAlert)
-{
-	using namespace TypedElementDataStorage;
-
-	if constexpr (bHasAlert)
-	{
-		Alert->CachedParent = InvalidRowHandle;
-	}
-	if constexpr (bHasChildAlert)
-	{
-		ChildAlert->bHasDecremented = true;
-	}
-}
-
-void UTypedElementAlertQueriesFactory::GuaranteeUpdateReentry(
-	TypedElementDataStorage::IQueryContext& Context, TypedElementDataStorage::RowHandle Row)
-{
-	// If not completing a full pass, add this tag to make sure a follow up pass happens. This is needed
-	// in case this update was triggered from a sync to/from external.
-	if (!Context.HasColumn<FTypedElementAlertUpdateTag>())
-	{
-		Context.AddColumns<FTypedElementAlertUpdateTag>(Row);
-	}
-}
-
-template<bool bHasAlert, bool bHasChildAlert>
-void UTypedElementAlertQueriesFactory::UpdateParent(FTypedElementAlertColumn* Alert, FTypedElementChildAlertColumn* ChildAlert,
-	TypedElementDataStorage::RowHandle NewParent)
-{
-	if constexpr (bHasAlert)
-	{
-		Alert->CachedParent = NewParent;
-	}
-	if constexpr (bHasChildAlert)
-	{
-		ChildAlert->CachedParent = NewParent;
-		ChildAlert->bHasDecremented = false;
-	}
-}
-
 void UTypedElementAlertQueriesFactory::AddChildAlertsToHierarchy(
 	TypedElementDataStorage::IQueryContext& Context, TypedElementDataStorage::RowHandle Parent, int32 ParentQueryIndex)
 {
 	using namespace TypedElementDataStorage;
 
 	bool bHasParent = true;
-	while (bHasParent)
+	do
 	{
-		RowHandle NextParent = Parent;
+		TypedElementDataStorage::RowHandle NextParent = Parent;
 		bHasParent = MoveToNextParent(NextParent, Context, ParentQueryIndex);
 
 		// Increment parent's child counter or create a child alert if one doesn't exist.
 		if (!Context.HasColumn<FTypedElementChildAlertColumn>(Parent))
 		{
 			FTypedElementChildAlertColumn ChildAlert;
+			ResetChildAlertCounters(ChildAlert);
 			ChildAlert.CachedParent = NextParent;
-			for (uint8 Index = 0; Index < static_cast<uint8>(FTypedElementAlertColumnType::MAX); ++Index)
-			{
-				ChildAlert.Counts[Index] = 0;
-			}
-			ChildAlert.RemoveCycleId = Context.GetUpdateCycleId();
-			ChildAlert.bHasDecremented = true;
 			Context.AddColumn(Parent, MoveTemp(ChildAlert));
 		}
-
 		Parent = NextParent;
-	}
+	} while (bHasParent);
 }
 
 void UTypedElementAlertQueriesFactory::IncrementParents(
-	TypedElementDataStorage::IQueryContext& Context, TypedElementDataStorage::RowHandle ParentRow,
-	uint16 ErrorIncrement, uint16 WarningIncrement,
+	TypedElementDataStorage::IQueryContext& Context, TypedElementDataStorage::RowHandle Row, FTypedElementAlertColumnType AlertType,
 	int32 ChildAlertQueryIndex)
 {
 	using namespace TypedElementDataStorage;
 	using namespace TypedElementQueryBuilder;
 	
-	// Assume that if the immediate parent has a child alert, then all of those parent's parents have a child alert as well.
-	while (true)
+	while (Context.IsRowAvailable(Row))
 	{
-		// Increment counter on the parent child alert.
-		FQueryResult Result = Context.RunSubquery(ChildAlertQueryIndex, ParentRow, CreateSubqueryCallbackBinding(
-			[&ParentRow, ErrorIncrement, WarningIncrement](
+		FQueryResult Result = Context.RunSubquery(ChildAlertQueryIndex, Row, CreateSubqueryCallbackBinding(
+			[&ParentRow = Row, AlertType](
 				ISubqueryContext& Context, RowHandle Row, FTypedElementChildAlertColumn& ChildAlert)
 			{
-				ChildAlert.Counts[static_cast<size_t>(FTypedElementAlertColumnType::Error)] += ErrorIncrement;
-				ChildAlert.Counts[static_cast<size_t>(FTypedElementAlertColumnType::Warning)] += WarningIncrement;
-				ParentRow = ChildAlert.CachedParent;
+				ChildAlert.Counts[static_cast<size_t>(AlertType)]++;
 				Context.AddColumns<FTypedElementSyncBackToWorldTag>(Row);
+				ParentRow = ChildAlert.CachedParent;
 			}));
 		checkf(Result.Count > 0, TEXT("Expected to be able to setup the child alert, but it was missing on the parent column."));
-		if (Result.Count == 0 || !Context.IsRowAvailable(ParentRow))
-		{
-			break;
-		}
 	}
 }
 
-void UTypedElementAlertQueriesFactory::DecrementParents(
-	TypedElementDataStorage::IQueryContext& Context, TypedElementDataStorage::RowHandle ParentRow,
-	uint16 ErrorDecrement, uint16 WarningDecrement,
-	int32 ChildAlertQueryIndex)
+void UTypedElementAlertQueriesFactory::ResetChildAlertCounters(FTypedElementChildAlertColumn& ChildAlert)
 {
-	using namespace TypedElementDataStorage;
-	using namespace TypedElementQueryBuilder;
-
-	while (true)
+	for (size_t It = 0; It < static_cast<size_t>(FTypedElementAlertColumnType::MAX); ++It)
 	{
-		FQueryResult Result = Context.RunSubquery(ChildAlertQueryIndex, ParentRow, CreateSubqueryCallbackBinding(
-			[&ParentRow, ErrorDecrement, WarningDecrement](
-				ISubqueryContext& Context, RowHandle Row, FTypedElementChildAlertColumn& ChildAlert)
-			{
-				ChildAlert.Counts[static_cast<size_t>(FTypedElementAlertColumnType::Error)] -= ErrorDecrement;
-				ChildAlert.Counts[static_cast<size_t>(FTypedElementAlertColumnType::Warning)] -= WarningDecrement;
-				if (ChildAlert.Counts[static_cast<size_t>(FTypedElementAlertColumnType::Error)] +
-					ChildAlert.Counts[static_cast<size_t>(FTypedElementAlertColumnType::Warning)] == 0)
-				{
-					Context.RemoveColumns<FTypedElementChildAlertColumn>(Row);
-				}
-				ParentRow = ChildAlert.CachedParent;
-				Context.AddColumns<FTypedElementSyncBackToWorldTag>(Row);
-			}));
-		if (Result.Count == 0 || !Context.IsRowAvailable(ParentRow))
-		{
-			break;
-		}
+		ChildAlert.Counts[It] = 0;
 	}
 }
 
diff --git a/Engine/Plugins/Experimental/TypedElementsDataStorage/Source/TypedElementsDataStorage/Private/Queries/TypedElementAlertQueries.h b/Engine/Plugins/Experimental/TypedElementsDataStorage/Source/TypedElementsDataStorage/Private/Queries/TypedElementAlertQueries.h
--- a/Engine/Plugins/Experimental/TypedElementsDataStorage/Source/TypedElementsDataStorage/Private/Queries/TypedElementAlertQueries.h
+++ b/Engine/Plugins/Experimental/TypedElementsDataStorage/Source/TypedElementsDataStorage/Private/Queries/TypedElementAlertQueries.h
@@ -2,6 +2,7 @@
 
 #pragma once
 
+#include "Elements/Columns/TypedElementAlertColumns.h"
 #include "Elements/Common/TypedElementHandles.h"
 #include "Elements/Interfaces/TypedElementDataStorageFactory.h"
 #include "Elements/Interfaces/TypedElementDataStorageInterface.h"
@@ -14,9 +15,6 @@ namespace TypedElementDataStorage
 {
 	struct IQueryContext;
 }
-struct FTypedElementAlertColumn;
-struct FTypedElementChildAlertColumn;
-struct FTypedElementParentColumn;
 
 /**
  * Calls to manage alerts, in particular child alerts.
@@ -33,70 +31,28 @@ public:
 
 private:
 	void RegisterSubQueries(ITypedElementDataStorageInterface& DataStorage);
-	void RegisterUpdateAlertsQueries(ITypedElementDataStorageInterface& DataStorage);
+	void RegisterParentUpdatesQueries(ITypedElementDataStorageInterface& DataStorage);
+	void RegisterChildAlertUpdatesQueries(ITypedElementDataStorageInterface& DataStorage);
 	void RegisterOnAddQueries(ITypedElementDataStorageInterface& DataStorage);
-	void RegisterOnAlertRemoveQueries(ITypedElementDataStorageInterface& DataStorage);
-	void RegisterOnParentRemoveQueries(ITypedElementDataStorageInterface& DataStorage);
-
-	template<bool bHasAlert, bool bHasChildAlert>
-	static void Update(TypedElementDataStorage::IQueryContext& Context, TypedElementDataStorage::RowHandle Row,
-		FTypedElementAlertColumn* Alert, FTypedElementChildAlertColumn* ChildAlert, const FTypedElementParentColumn& Parent,
-		int32 ChildAlertColumnReadWriteSubquery, int32 ParentReadOnlySubquery);
-	
-	template<bool bHasAlert, bool bHasChildAlert>
-	static void NeedsUpdating(FTypedElementAlertColumn* Alert, FTypedElementChildAlertColumn* ChildAlert,
-		const FTypedElementParentColumn& Parent, bool& bAlertNeedsUpdate, bool& bChildAlertsNeedUpdate);
-	
-	template<bool bHasAlert, bool bHasChildAlert>
-	static void NeedsDecrementing(FTypedElementAlertColumn* Alert, FTypedElementChildAlertColumn* ChildAlert,
-		const FTypedElementParentColumn& Parent, bool& bAlertNeedsDecrementing, bool& bChildAlertsNeedDecrementing);
-	
-	template<bool bHasAlert, bool bHasChildAlert>
-	static TypedElementDataStorage::RowHandle GetParent(FTypedElementAlertColumn* Alert, FTypedElementChildAlertColumn* ChildAlert);
-	
-	template<bool bHasAlert, bool bHasChildAlert>
-	static void GetTotalCounts(FTypedElementAlertColumn* Alert, FTypedElementChildAlertColumn* ChildAlert,
-		uint16& ErrorCount, uint16& WarningCount, bool IncludeAlert, bool IncludeChildAlerts);
-	
-	template<bool bHasAlert, bool bHasChildAlert>
-	static void MarkDecremented(FTypedElementAlertColumn* Alert, FTypedElementChildAlertColumn* ChildAlert);
-	
-	static void GuaranteeUpdateReentry(TypedElementDataStorage::IQueryContext& Context, TypedElementDataStorage::RowHandle Row);
-	
-	template<bool bHasAlert, bool bHasChildAlert>
-	static void UpdateParent(FTypedElementAlertColumn* Alert, FTypedElementChildAlertColumn* ChildAlert,
-		TypedElementDataStorage::RowHandle NewParent);
+	void RegisterOnRemoveQueries(ITypedElementDataStorageInterface& DataStorage);
 
 	static void AddChildAlertsToHierarchy(
 		TypedElementDataStorage::IQueryContext& Context, TypedElementDataStorage::RowHandle Parent, int32 ParentQueryIndex);
 
 	static void IncrementParents(
-		TypedElementDataStorage::IQueryContext& Context, TypedElementDataStorage::RowHandle ParentRow,
-		uint16 ErrorIncrement, uint16 WarningIncrement,
-		int32 ChildAlertQueryIndex);
-	static void DecrementParents(
-		TypedElementDataStorage::IQueryContext& Context, TypedElementDataStorage::RowHandle ParentRow,
-		uint16 ErrorDecrement, uint16 WarningDecrement,
+		TypedElementDataStorage::IQueryContext& Context, TypedElementDataStorage::RowHandle Row, FTypedElementAlertColumnType AlertType,
 		int32 ChildAlertQueryIndex);
+	
+	static void ResetChildAlertCounters(FTypedElementChildAlertColumn& ChildAlert);
 
-	// Find the next parent to the current parent. If there's no parent this function return false. If there is a parent, it will return
-	// true and Parent will be updated to the new parent.
 	static bool MoveToNextParent(
 		TypedElementDataStorage::RowHandle& Parent, TypedElementDataStorage::IQueryContext& Context, int32 SubQueryIndex);
 
+	static const FName AlertConditionName;
 	TypedElementDataStorage::QueryHandle ChildAlertColumnReadWriteQuery;
 	TypedElementDataStorage::QueryHandle ParentReadOnlyQuery;
 };
 
-/**
- * Tag used during the update cycle as this can take between 1 and 3 cycles to complete.
- */
-USTRUCT(meta = (DisplayName = "Alert update"))
-struct FTypedElementAlertUpdateTag final : public FTypedElementDataStorageTag
-{
-	GENERATED_BODY()
-};
-
 /**
  * Enable AlertColumn to be mementoized
  */
diff --git a/Engine/Source/Runtime/TypedElementFramework/Public/Elements/Columns/TypedElementAlertColumns.h b/Engine/Source/Runtime/TypedElementFramework/Public/Elements/Columns/TypedElementAlertColumns.h
--- a/Engine/Source/Runtime/TypedElementFramework/Public/Elements/Columns/TypedElementAlertColumns.h
+++ b/Engine/Source/Runtime/TypedElementFramework/Public/Elements/Columns/TypedElementAlertColumns.h
@@ -32,11 +32,7 @@ struct FTypedElementAlertColumn final : public FTypedElementDataStorageColumn
 	// Store a copy of the parent row so it's possible to detect if a row has been reparented.
 	TypedElementDataStorage::RowHandle CachedParent;
 
-	// The cycle id is 64 bits, but this column is only interested in avoiding the same update happening multiple times in the same frame
-	// therefore this is kept small to stay within the padding of the struct. This could even be reduced to a 8 bit value if needed.
-	uint16 RemoveCycleId;
-
-	UPROPERTY()
+	UPROPERTY(meta = (IgnoreForMemberInitializationTest))
 	FTypedElementAlertColumnType AlertType;
 };
 
@@ -52,14 +48,6 @@ struct FTypedElementChildAlertColumn final : public FTypedElementDataStorageColu
 	TypedElementDataStorage::RowHandle CachedParent;
 
 	uint16 Counts[static_cast<size_t>(FTypedElementAlertColumnType::MAX)];
-
-	// The cycle id is 64 bits, but this column is only interested in avoiding the same update happening multiple times in the same frame
-	// therefore this is kept small to stay within the padding of the struct. This could even be reduced to a 8 bit value if needed.
-	uint16 RemoveCycleId;
-
-	// Indicates if during updating recently, this column has already decremented its parents. This is only needed to avoid repeated
-	// decrements if updating the child alerts takes more than a single frame. It's not used outside updating.
-	bool bHasDecremented;
 };
 
 /**
