[Revision Control] Remove overhead of sync to world tag and update query conditions to include optional columns

diff --git a/Engine/Plugins/Developer/PerforceSourceControl/Source/PerforceSourceControl/Private/PerforceSourceControlOperations.cpp b/Engine/Plugins/Developer/PerforceSourceControl/Source/PerforceSourceControl/Private/PerforceSourceControlOperations.cpp
--- a/Engine/Plugins/Developer/PerforceSourceControl/Source/PerforceSourceControl/Private/PerforceSourceControlOperations.cpp
+++ b/Engine/Plugins/Developer/PerforceSourceControl/Source/PerforceSourceControl/Private/PerforceSourceControlOperations.cpp
@@ -2067,7 +2067,7 @@ bool FPerforceUpdateStatusWorker::UpdateStates() const
 		ExternalRevisionId.RevisionId.Id[0] = State.DepotRevNumber;
 		DataStorage->AddOrGetColumn(Row, MoveTemp(ExternalRevisionId));
 
-		TArray<UScriptStruct*> ToAdd { FTypedElementSyncFromWorldTag::StaticStruct() };
+		TArray<UScriptStruct*> ToAdd;
 		TArray<UScriptStruct*> ToRemove;
 		auto SyncTagFromState = [&](bool bCondition, UScriptStruct* Tag)
 		{
diff --git a/Engine/Plugins/Experimental/TypedElementsDataStorage/Source/TypedElementsDataStorageRevisionControl/Source/Private/RevisionControlProcessors.cpp b/Engine/Plugins/Experimental/TypedElementsDataStorage/Source/TypedElementsDataStorageRevisionControl/Source/Private/RevisionControlProcessors.cpp
--- a/Engine/Plugins/Experimental/TypedElementsDataStorage/Source/TypedElementsDataStorageRevisionControl/Source/Private/RevisionControlProcessors.cpp
+++ b/Engine/Plugins/Experimental/TypedElementsDataStorage/Source/TypedElementsDataStorageRevisionControl/Source/Private/RevisionControlProcessors.cpp
@@ -135,7 +135,7 @@ void UTypedElementRevisionControlFactory::RegisterQueries(ITypedElementDataStora
 	TypedElementQueryHandle ObjectToSCCQuery = DataStorage.RegisterQuery(
 		Select()
 			.ReadOnly<FTypedElementPackagePathColumn>()
-			.ReadOnly<FSCCStatusColumn>(EOptional::Yes)
+			.ReadOnly<FSCCStatusColumn, FSCCExternallyLockedColumn>(EOptional::Yes)
 		.Compile());
 
 	DataStorage.RegisterQuery(
@@ -163,6 +163,8 @@ void UTypedElementRevisionControlFactory::RegisterQueries(ITypedElementDataStora
 				);
 			}
 		)
+		.Where()
+			.Any<FSCCNotCurrentTag, FSCCLockedTag>()
 		.DependsOn()
 			.SubQuery(ObjectToSCCQuery)
 		.Compile()
