Updated the direct TEDS API to fully support delayed column operations.

Previously only some TEDS operations done through the direct API would be queued up if the target row hadn't been created yet. As a result it was possible to delay add columns without data, but if one with data was added it would simply disappear. Now that data is stored in the scratch buffer and applied later. To allow the direct API to still appear to function as if the operations were immediately executed, the command buffer keeps track of the future state of the columns.

The following additional changers were made:
- Several column functions that returned `bool` are now returning `void`. The results were never used and most were hard coded to return true.
- Moved the templated column functions up so they're next to the regular functions in the interface. Having them located below made them harder to discover.
- All the overloaded functions that used `FTopLevelAssetPath` have been removed. It's going to be easier to have a higher level layer deal with the translation from string to column type. This keeps the interface cleaner.
- `AddOrGetColumn` has been updated with the ability to provide callbacks from initializing and relocating columns. This gives an opportunity to initialize a column before FObserver::OnAdd is called and the option to move instead of copy the column out of the scratch buffer. The templated version of `AddOrGetColumn` makes full use of this.
- The `AddOrGetColumn` overload that takes arbitrary arguments has been removed. This was done to avoid potential slicing, further discourage creating constructors in columns and to make sure that FObserver::OnAdd is called before initialization.
- The command buffer for the direct API has been moved into the environment. This was done to make it more widely available and so its lifecyle is more closely tied to the scratch buffer and Mass objects.
- Moved the concepts and the (new) callbacks to a central location so they're easier to use.
- Introduced concepts to the `AddColumn` calls to avoid to use `enable_if` to keep `AddColumn(..., ColumnType&&)` and `AddColumn(..., const UScriptStruct*)` from conflicting.

diff --git a/Engine/Plugins/Experimental/TypedElementsDataStorage/Source/TypedElementsDataStorage/Private/TypedElementDatabase.cpp b/Engine/Plugins/Experimental/TypedElementsDataStorage/Source/TypedElementsDataStorage/Private/TypedElementDatabase.cpp
--- a/Engine/Plugins/Experimental/TypedElementsDataStorage/Source/TypedElementsDataStorage/Private/TypedElementDatabase.cpp
+++ b/Engine/Plugins/Experimental/TypedElementsDataStorage/Source/TypedElementsDataStorage/Private/TypedElementDatabase.cpp
@@ -240,8 +240,7 @@ void UTypedElementDatabase::OnPreMassTick(float DeltaTime)
 	OnUpdateDelegate.Broadcast();
 	// Process pending commands after other systems have had a chance to update. Other systems may have executed work needed
 	// to complete pending work.
-	FTypedElementDatabaseCommandBuffer::ProcessBuffer(DeferredCommands, *Environment);
-	DeferredCommands.Reset();
+	Environment->GetDirectDeferredCommands().ProcessCommands();
 }
 
 void UTypedElementDatabase::OnPostMassTick(float DeltaTime)
@@ -330,7 +329,7 @@ bool UTypedElementDatabase::AddRow(TypedElementRowHandle ReservedRow, TypedEleme
 	}
 }
 
-bool UTypedElementDatabase::BatchAddRow(TypedElementTableHandle Table, int32 Count, TypedElementDataStorageCreationCallbackRef OnCreated)
+bool UTypedElementDatabase::BatchAddRow(TypedElementTableHandle Table, int32 Count, TypedElementDataStorage::RowCreationCallbackRef OnCreated)
 {
 	OnCreated.CheckCallable();
 	checkf(Table < Tables.Num(), TEXT("Attempting to add multiple rows to a non-existing table."));
@@ -352,7 +351,7 @@ bool UTypedElementDatabase::BatchAddRow(TypedElementTableHandle Table, int32 Cou
 }
 
 bool UTypedElementDatabase::BatchAddRow(TypedElementTableHandle Table, TConstArrayView<TypedElementRowHandle> ReservedHandles,
-	TypedElementDataStorageCreationCallbackRef OnCreated)
+	TypedElementDataStorage::RowCreationCallbackRef OnCreated)
 {
 	OnCreated.CheckCallable();
 	checkf(Table < Tables.Num(), TEXT("Attempting to add multiple rows to a non-existing table."));
@@ -393,6 +392,7 @@ void UTypedElementDatabase::RemoveRow(TypedElementRowHandle Row)
 		}
 		else
 		{
+			Environment->GetDirectDeferredCommands().Clear(Row);
 			ActiveEditorEntityManager->ReleaseReservedEntity(FMassEntityHandle::FromNumber(Row));
 		}
 	}
@@ -408,100 +408,94 @@ bool UTypedElementDatabase::HasRowBeenAssigned(TypedElementRowHandle Row) const
 	return ActiveEditorEntityManager ? FTypedElementDatabaseCommandBuffer::Execute_HasRowBeenAssigned(*ActiveEditorEntityManager, Row) : false;
 }
 
-bool UTypedElementDatabase::AddColumn(TypedElementRowHandle Row, const UScriptStruct* ColumnType)
+void UTypedElementDatabase::AddColumn(TypedElementRowHandle Row, const UScriptStruct* ColumnType)
 {
-	if (ColumnType && ActiveEditorEntityManager && HasRowBeenAssigned(Row))
+	if (ColumnType && ActiveEditorEntityManager)
 	{
-		FTypedElementDatabaseCommandBuffer::Execute_AddColumnCommand(*ActiveEditorEntityManager, Row, ColumnType);
-	}
-	else
-	{
-		FTypedElementDatabaseCommandBuffer::Queue_AddColumnCommand(DeferredCommands, Row, ColumnType);
+		if (HasRowBeenAssigned(Row))
+		{
+			FTypedElementDatabaseCommandBuffer::Execute_AddColumnCommand(*ActiveEditorEntityManager, Row, ColumnType);
+		}
+		else
+		{
+			Environment->GetDirectDeferredCommands().Queue_AddColumnCommand(Row, ColumnType);
+		}
 	}
-	return true;
 }
 
-bool UTypedElementDatabase::AddColumn(TypedElementRowHandle Row, FTopLevelAssetPath ColumnName)
+void UTypedElementDatabase::AddColumnData(TypedElementRowHandle Row, const UScriptStruct* ColumnType,
+	const TypedElementDataStorage::ColumnCreationCallbackRef& Initializer,
+	TypedElementDataStorage::ColumnCopyOrMoveCallback Relocator)
 {
-	bool bExactMatch = true;
-	UScriptStruct* ColumnStructInfo = Cast<UScriptStruct>(StaticFindObject(UScriptStruct::StaticClass(), ColumnName, bExactMatch));
-	return ColumnStructInfo ? AddColumn(Row, ColumnStructInfo) : false;
+	AddOrGetColumnData(Row, ColumnType, Initializer, Relocator);
 }
 
 void UTypedElementDatabase::RemoveColumn(TypedElementRowHandle Row, const UScriptStruct* ColumnType)
 {
-	if (ActiveEditorEntityManager && HasRowBeenAssigned(Row))
-	{
-		FTypedElementDatabaseCommandBuffer::Execute_RemoveColumnCommand(*ActiveEditorEntityManager, Row, ColumnType);
-	}
-	else
-	{
-		FTypedElementDatabaseCommandBuffer::Queue_RemoveColumnCommand(DeferredCommands, Row, ColumnType);
-	}
-}
-
-void UTypedElementDatabase::RemoveColumn(TypedElementRowHandle Row, FTopLevelAssetPath ColumnName)
-{
-	bool bExactMatch = true;
-	if (UScriptStruct* ColumnStructInfo = Cast<UScriptStruct>(StaticFindObject(UScriptStruct::StaticClass(), ColumnName, bExactMatch)))
+	if (ColumnType && ActiveEditorEntityManager)
 	{
-		RemoveColumn(Row, ColumnStructInfo);
+		if (HasRowBeenAssigned(Row))
+		{
+			FTypedElementDatabaseCommandBuffer::Execute_RemoveColumnCommand(*ActiveEditorEntityManager, Row, ColumnType);
+		}
+		else
+		{
+			Environment->GetDirectDeferredCommands().Queue_RemoveColumnCommand(Row, ColumnType);
+		}
 	}
 }
 
-void* UTypedElementDatabase::AddOrGetColumnData(TypedElementRowHandle Row, const UScriptStruct* ColumnType)
+void* UTypedElementDatabase::AddOrGetColumnData(TypedElementRowHandle Row, const UScriptStruct* ColumnType,
+	const TypedElementDataStorage::ColumnCreationCallbackRef& Initializer, TypedElementDataStorage::ColumnCopyOrMoveCallback Relocator)
 {
 	FMassEntityHandle Entity = FMassEntityHandle::FromNumber(Row);
-	if (ActiveEditorEntityManager && ActiveEditorEntityManager->IsEntityActive(Entity) &&
-		ColumnType && ColumnType->IsChildOf(FMassFragment::StaticStruct()))
+	if (ActiveEditorEntityManager && ColumnType && ColumnType->IsChildOf(FMassFragment::StaticStruct()))
 	{
-		FStructView Column = ActiveEditorEntityManager->GetFragmentDataStruct(Entity, ColumnType);
-		if (!Column.IsValid())
+		if (HasRowBeenAssigned(Row))
 		{
-			ActiveEditorEntityManager->AddFragmentToEntity(Entity, ColumnType);
-			Column = ActiveEditorEntityManager->GetFragmentDataStruct(Entity, ColumnType);
-			checkf(Column.IsValid(), TEXT("Added a new column to the Typed Element's data store, but it couldn't be retrieved."));
+			FStructView Column = ActiveEditorEntityManager->GetFragmentDataStruct(Entity, ColumnType);
+			if (!Column.IsValid())
+			{
+				ActiveEditorEntityManager->AddFragmentToEntity(Entity, ColumnType, Initializer);
+				Column = ActiveEditorEntityManager->GetFragmentDataStruct(Entity, ColumnType);
+				checkf(Column.IsValid(), TEXT("Added a new column to the Typed Element's data store, but it couldn't be retrieved."));
+				return Column.GetMemory();
 
+			}
+			else
+			{
+				void* Result = Column.GetMemory();
+				Initializer(Result, *ColumnType);
+				return Result;
+			}
+		}
+		else
+		{
+			void* Result = Environment->GetDirectDeferredCommands().Queue_AddDataColumnCommandUnitialized(Row, ColumnType, Relocator);
+			Initializer(Result, *ColumnType);
+			return Result;
 		}
-		return Column.GetMemory();
 	}
 	return nullptr;
 }
 
-ColumnDataResult UTypedElementDatabase::AddOrGetColumnData(TypedElementRowHandle Row, FTopLevelAssetPath ColumnName)
-{
-	constexpr bool bExactMatch = true;
-	UScriptStruct* FragmentStructInfo = Cast<UScriptStruct>(StaticFindObject(UScriptStruct::StaticClass(), ColumnName, bExactMatch));
-	return FragmentStructInfo ?
-		ColumnDataResult{ FragmentStructInfo, AddOrGetColumnData(Row, FragmentStructInfo) }:
-		ColumnDataResult{ nullptr, nullptr };
-}
-
-ColumnDataResult UTypedElementDatabase::AddOrGetColumnData(TypedElementRowHandle Row, FTopLevelAssetPath ColumnName,
-	TConstArrayView<TypedElement::ColumnUtils::Argument> Arguments)
-{
-	ColumnDataResult Result = AddOrGetColumnData(Row, ColumnName);
-	if (Result.Description && Result.Data)
-	{
-		TypedElement::ColumnUtils::SetColumnValues(Result.Data, Result.Description, Arguments);
-		return Result;
-	}
-	else
-	{
-		return ColumnDataResult{ nullptr, nullptr };
-	}
-}
-
 const void* UTypedElementDatabase::GetColumnData(TypedElementRowHandle Row, const UScriptStruct* ColumnType) const
 {
 	FMassEntityHandle Entity = FMassEntityHandle::FromNumber(Row);
-	if (ActiveEditorEntityManager && ActiveEditorEntityManager->IsEntityActive(Entity) &&
+	if (ActiveEditorEntityManager &&
 		ColumnType && ColumnType->IsChildOf(FMassFragment::StaticStruct()))
 	{
-		FStructView Column = ActiveEditorEntityManager->GetFragmentDataStruct(Entity, ColumnType);
-		if (Column.IsValid())
+		if (ActiveEditorEntityManager->IsEntityActive(Entity))
+		{
+			FStructView Column = ActiveEditorEntityManager->GetFragmentDataStruct(Entity, ColumnType);
+			if (Column.IsValid())
+			{
+				return Column.GetMemory();
+			}
+		}
+		else
 		{
-			return Column.GetMemory();
+			return Environment->GetDirectDeferredCommands().GetQueuedDataColumn(Row, ColumnType);
 		}
 	}
 	return nullptr;
@@ -512,35 +506,7 @@ void* UTypedElementDatabase::GetColumnData(TypedElementRowHandle Row, const UScr
 	return const_cast<void*>(static_cast<const UTypedElementDatabase*>(this)->GetColumnData(Row, ColumnType));
 }
 
-ColumnDataResult UTypedElementDatabase::GetColumnData(TypedElementRowHandle Row, FTopLevelAssetPath ColumnName)
-{
-	FMassEntityHandle Entity = FMassEntityHandle::FromNumber(Row);
-	if (ActiveEditorEntityManager && ActiveEditorEntityManager->IsEntityActive(Entity))
-	{
-		const UScriptStruct* FragmentType = nullptr;
-		FMassArchetypeHandle Archetype = ActiveEditorEntityManager->GetArchetypeForEntityUnsafe(Entity);
-		ActiveEditorEntityManager->ForEachArchetypeFragmentType(Archetype, 
-			[ColumnName, &FragmentType](const UScriptStruct* Fragment)
-			{
-				if (Fragment->GetStructPathName() == ColumnName)
-				{
-					FragmentType = Fragment;
-				}
-			});
-
-		if (FragmentType && FragmentType->IsChildOf(FMassFragment::StaticStruct()))
-		{
-			FStructView Column = ActiveEditorEntityManager->GetFragmentDataStruct(Entity, FragmentType);
-			if (Column.IsValid())
-			{
-				return ColumnDataResult{ FragmentType, Column.GetMemory() };
-			}
-		}
-	}
-	return ColumnDataResult{ nullptr, nullptr };
-}
-
-bool UTypedElementDatabase::AddColumns(TypedElementRowHandle Row, TConstArrayView<const UScriptStruct*> Columns)
+void UTypedElementDatabase::AddColumns(TypedElementRowHandle Row, TConstArrayView<const UScriptStruct*> Columns)
 {
 	if (ActiveEditorEntityManager)
 	{
@@ -557,12 +523,10 @@ bool UTypedElementDatabase::AddColumns(TypedElementRowHandle Row, TConstArrayVie
 			}
 			else
 			{
-				FTypedElementDatabaseCommandBuffer::Queue_AddColumnsCommand(DeferredCommands, Row, FragmentsToAdd, TagsToAdd);
+				Environment->GetDirectDeferredCommands().Queue_AddColumnsCommand(Row, FragmentsToAdd, TagsToAdd);
 			}
-			return true;
 		}
 	}
-	return false;
 }
 
 void UTypedElementDatabase::RemoveColumns(TypedElementRowHandle Row, TConstArrayView<const UScriptStruct*> Columns)
@@ -582,45 +546,54 @@ void UTypedElementDatabase::RemoveColumns(TypedElementRowHandle Row, TConstArray
 			}
 			else
 			{
-				FTypedElementDatabaseCommandBuffer::Queue_RemoveColumnsCommand(DeferredCommands, Row, FragmentsToRemove, TagsToRemove);
+				Environment->GetDirectDeferredCommands().Queue_RemoveColumnsCommand(Row, FragmentsToRemove, TagsToRemove);
 			}
 		}
 	}
 }
 
-bool UTypedElementDatabase::AddRemoveColumns(TypedElementRowHandle Row,
+void UTypedElementDatabase::AddRemoveColumns(TypedElementRowHandle Row,
 	TConstArrayView<const UScriptStruct*> ColumnsToAdd, TConstArrayView<const UScriptStruct*> ColumnsToRemove)
 {
-	bool bResult = false;
-	FMassEntityHandle Entity = FMassEntityHandle::FromNumber(Row);
-	if (ActiveEditorEntityManager && ActiveEditorEntityManager->IsEntityActive(Entity))
+	if (ActiveEditorEntityManager)
 	{
+		FMassEntityHandle Entity = FMassEntityHandle::FromNumber(Row);
 		FMassArchetypeHandle Archetype = ActiveEditorEntityManager->GetArchetypeForEntity(Entity);
 
 		FMassFragmentBitSet FragmentsToAdd;
 		FMassTagBitSet TagsToAdd;
-		if (TypedElementDatabasePrivate::ColumnsToBitSets(ColumnsToAdd, FragmentsToAdd, TagsToAdd).MustUpdate())
-		{
-			FMassArchetypeCompositionDescriptor AddComposition(
-				MoveTemp(FragmentsToAdd), MoveTemp(TagsToAdd), FMassChunkFragmentBitSet(), FMassSharedFragmentBitSet());
-			ActiveEditorEntityManager->AddCompositionToEntity_GetDelta(Entity, AddComposition);
-			bResult = true;
-		}
-
 		FMassTagBitSet TagsToRemove;
 		FMassFragmentBitSet FragmentsToRemove;
-		if (TypedElementDatabasePrivate::ColumnsToBitSets(ColumnsToRemove, FragmentsToRemove, TagsToRemove).MustUpdate())
+
+		bool bMustAddColumns = TypedElementDatabasePrivate::ColumnsToBitSets(ColumnsToAdd, FragmentsToAdd, TagsToAdd).MustUpdate();
+		bool bMustRemoveColumns = TypedElementDatabasePrivate::ColumnsToBitSets(ColumnsToRemove, FragmentsToRemove, TagsToRemove).MustUpdate();
+		
+		if (ActiveEditorEntityManager->IsEntityActive(Entity))
+		{
+			if (bMustAddColumns)
+			{
+				FTypedElementDatabaseCommandBuffer::Execute_AddColumnsCommand(*ActiveEditorEntityManager, Row, FragmentsToAdd, TagsToAdd);
+			}
+			if (bMustRemoveColumns)
+			{
+				FTypedElementDatabaseCommandBuffer::Execute_RemoveColumnsCommand(*ActiveEditorEntityManager, Row, FragmentsToRemove, TagsToRemove);
+			}
+		}
+		else
 		{
-			FMassArchetypeCompositionDescriptor RemoveComposition(
-				MoveTemp(FragmentsToRemove), MoveTemp(TagsToRemove), FMassChunkFragmentBitSet(), FMassSharedFragmentBitSet());
-			ActiveEditorEntityManager->RemoveCompositionFromEntity(Entity, RemoveComposition);
-			bResult = true;
+			if (bMustAddColumns)
+			{
+				Environment->GetDirectDeferredCommands().Queue_AddColumnsCommand(Row, FragmentsToAdd, TagsToAdd);
+			}
+			if (bMustRemoveColumns)
+			{
+				Environment->GetDirectDeferredCommands().Queue_RemoveColumnsCommand(Row, FragmentsToRemove, TagsToRemove);
+			}
 		}
 	}
-	return bResult;
 }
 
-bool UTypedElementDatabase::BatchAddRemoveColumns(TConstArrayView<TypedElementRowHandle> Rows, 
+void UTypedElementDatabase::BatchAddRemoveColumns(TConstArrayView<TypedElementRowHandle> Rows, 
 	TConstArrayView<const UScriptStruct*> ColumnsToAdd, TConstArrayView<const UScriptStruct*> ColumnsToRemove)
 {	
 	if (ActiveEditorEntityManager)
@@ -642,6 +615,8 @@ bool UTypedElementDatabase::BatchAddRemoveColumns(TConstArrayView<TypedElementRo
 			using EntityArchetypeLookup = TMap<FMassArchetypeHandle, EntityHandleArray, TInlineSetAllocator<32>>;
 			using ArchetypeEntityArray = TArray<FMassArchetypeEntityCollection, TInlineAllocator<32>>;
 
+			FTypedElementDatabaseCommandBuffer& CommandBuffer = Environment->GetDirectDeferredCommands();
+			
 			// Sort rows (entities) into to matching table (archetype) bucket.
 			EntityArchetypeLookup LookupTable;
 			for (TypedElementRowHandle EntityId : Rows)
@@ -653,6 +628,17 @@ bool UTypedElementDatabase::BatchAddRemoveColumns(TConstArrayView<TypedElementRo
 					EntityHandleArray& EntityCollection = LookupTable.FindOrAdd(Archetype);
 					EntityCollection.Add(Entity);
 				}
+				else
+				{
+					if (AddResult.MustUpdate())
+					{
+						CommandBuffer.Queue_AddColumnsCommand(EntityId, FragmentsToAdd, TagsToAdd);
+					}
+					if (RemoveResult.MustUpdate())
+					{
+						CommandBuffer.Queue_RemoveColumnsCommand(EntityId, FragmentsToRemove, TagsToRemove);
+					}
+				}
 			}
 		
 			// Construct table (archetype) specific row (entity) collections.
@@ -672,35 +658,46 @@ bool UTypedElementDatabase::BatchAddRemoveColumns(TConstArrayView<TypedElementRo
 			{
 				ActiveEditorEntityManager->BatchChangeTagsForEntities(EntityCollections, TagsToAdd, TagsToRemove);
 			}
-			return true;
 		}
 	}
-	return false;
 }
 
 bool UTypedElementDatabase::HasColumns(TypedElementRowHandle Row, TConstArrayView<const UScriptStruct*> ColumnTypes) const
 {
-	FMassEntityHandle Entity = FMassEntityHandle::FromNumber(Row);
-	if (ActiveEditorEntityManager && ActiveEditorEntityManager->IsEntityActive(Entity))
+	if (ActiveEditorEntityManager)
 	{
-		FMassArchetypeHandle Archetype = ActiveEditorEntityManager->GetArchetypeForEntity(Entity);
-		const FMassArchetypeCompositionDescriptor& Composition = ActiveEditorEntityManager->GetArchetypeComposition(Archetype);
-
 		bool bHasAllColumns = true;
-		const UScriptStruct* const* ColumnTypesEnd = ColumnTypes.end();
-		for (const UScriptStruct* const* ColumnType = ColumnTypes.begin(); ColumnType != ColumnTypesEnd && bHasAllColumns; ++ColumnType)
+		
+		FMassEntityHandle Entity = FMassEntityHandle::FromNumber(Row);
+		if (ActiveEditorEntityManager->IsEntityActive(Entity))
 		{
-			if ((*ColumnType)->IsChildOf(FMassFragment::StaticStruct()))
-			{
-				bHasAllColumns = Composition.Fragments.Contains(**ColumnType);
-			}
-			else if ((*ColumnType)->IsChildOf(FMassTag::StaticStruct()))
+			FMassArchetypeHandle Archetype = ActiveEditorEntityManager->GetArchetypeForEntity(Entity);
+			const FMassArchetypeCompositionDescriptor& Composition = ActiveEditorEntityManager->GetArchetypeComposition(Archetype);
+
+			const UScriptStruct* const* ColumnTypesEnd = ColumnTypes.end();
+			for (const UScriptStruct* const* ColumnType = ColumnTypes.begin(); ColumnType != ColumnTypesEnd && bHasAllColumns; ++ColumnType)
 			{
-				bHasAllColumns = Composition.Tags.Contains(**ColumnType);
+				if ((*ColumnType)->IsChildOf(FMassFragment::StaticStruct()))
+				{
+					bHasAllColumns = Composition.Fragments.Contains(**ColumnType);
+				}
+				else if ((*ColumnType)->IsChildOf(FMassTag::StaticStruct()))
+				{
+					bHasAllColumns = Composition.Tags.Contains(**ColumnType);
+				}
+				else
+				{
+					return false;
+				}
 			}
-			else
+		}
+		else
+		{
+			const FTypedElementDatabaseCommandBuffer& CommandBuffer = Environment->GetDirectDeferredCommands();
+			const UScriptStruct* const* ColumnTypesEnd = ColumnTypes.end();
+			for (const UScriptStruct* const* ColumnType = ColumnTypes.begin(); ColumnType != ColumnTypesEnd && bHasAllColumns; ++ColumnType)
 			{
-				return false;
+				bHasAllColumns = CommandBuffer.HasColumn(Row, *ColumnType);
 			}
 		}
 
@@ -711,30 +708,43 @@ bool UTypedElementDatabase::HasColumns(TypedElementRowHandle Row, TConstArrayVie
 
 bool UTypedElementDatabase::HasColumns(TypedElementRowHandle Row, TConstArrayView<TWeakObjectPtr<const UScriptStruct>> ColumnTypes) const
 {
-	FMassEntityHandle Entity = FMassEntityHandle::FromNumber(Row);
-	if (ActiveEditorEntityManager && ActiveEditorEntityManager->IsEntityActive(Entity))
+	if (ActiveEditorEntityManager)
 	{
-		FMassArchetypeHandle Archetype = ActiveEditorEntityManager->GetArchetypeForEntity(Entity);
-		const FMassArchetypeCompositionDescriptor& Composition = ActiveEditorEntityManager->GetArchetypeComposition(Archetype);
-
 		bool bHasAllColumns = true;
-		const TWeakObjectPtr<const UScriptStruct>* ColumnTypesEnd = ColumnTypes.end();
-		for (const TWeakObjectPtr<const UScriptStruct>* ColumnType = ColumnTypes.begin(); ColumnType != ColumnTypesEnd && bHasAllColumns; ++ColumnType)
+	
+		FMassEntityHandle Entity = FMassEntityHandle::FromNumber(Row);
+		if (ActiveEditorEntityManager->IsEntityActive(Entity))
 		{
-			if (ColumnType->IsValid())
+			FMassArchetypeHandle Archetype = ActiveEditorEntityManager->GetArchetypeForEntity(Entity);
+			const FMassArchetypeCompositionDescriptor& Composition = ActiveEditorEntityManager->GetArchetypeComposition(Archetype);
+
+			const TWeakObjectPtr<const UScriptStruct>* ColumnTypesEnd = ColumnTypes.end();
+			for (const TWeakObjectPtr<const UScriptStruct>* ColumnType = ColumnTypes.begin(); ColumnType != ColumnTypesEnd && bHasAllColumns; ++ColumnType)
 			{
-				if ((*ColumnType)->IsChildOf(FMassFragment::StaticStruct()))
-				{
-					bHasAllColumns = Composition.Fragments.Contains(**ColumnType);
-					continue;
-				}
-				else if ((*ColumnType)->IsChildOf(FMassTag::StaticStruct()))
+				if (ColumnType->IsValid())
 				{
-					bHasAllColumns = Composition.Tags.Contains(**ColumnType);
-					continue;
+					if ((*ColumnType)->IsChildOf(FMassFragment::StaticStruct()))
+					{
+						bHasAllColumns = Composition.Fragments.Contains(**ColumnType);
+						continue;
+					}
+					else if ((*ColumnType)->IsChildOf(FMassTag::StaticStruct()))
+					{
+						bHasAllColumns = Composition.Tags.Contains(**ColumnType);
+						continue;
+					}
 				}
+				return false;
+			}
+		}
+		else
+		{
+			const FTypedElementDatabaseCommandBuffer& CommandBuffer = Environment->GetDirectDeferredCommands();
+			const TWeakObjectPtr<const UScriptStruct>* ColumnTypesEnd = ColumnTypes.end();
+			for (const TWeakObjectPtr<const UScriptStruct>* ColumnType = ColumnTypes.begin(); ColumnType != ColumnTypesEnd && bHasAllColumns; ++ColumnType)
+			{
+				bHasAllColumns = CommandBuffer.HasColumn(Row, ColumnType->Get());
 			}
-			return false;
 		}
 
 		return bHasAllColumns;
@@ -744,28 +754,44 @@ bool UTypedElementDatabase::HasColumns(TypedElementRowHandle Row, TConstArrayVie
 
 bool UTypedElementDatabase::MatchesColumns(TypedElementDataStorage::RowHandle Row, const TypedElementDataStorage::FQueryConditions& Conditions) const
 {
-	FMassEntityHandle Entity = FMassEntityHandle::FromNumber(Row);
-	if (ActiveEditorEntityManager && ActiveEditorEntityManager->IsEntityActive(Entity))
+	if (ActiveEditorEntityManager)
 	{
-		FMassArchetypeHandle Archetype = ActiveEditorEntityManager->GetArchetypeForEntity(Entity);
-		const FMassArchetypeCompositionDescriptor& Composition = ActiveEditorEntityManager->GetArchetypeComposition(Archetype);
-
-		auto Callback = [&Composition](uint8_t ColumnIndex, TWeakObjectPtr<const UScriptStruct> Column)
+		FMassEntityHandle Entity = FMassEntityHandle::FromNumber(Row);
+		if (ActiveEditorEntityManager->IsEntityActive(Entity))
 		{
-			if (Column.IsValid())
+			FMassArchetypeHandle Archetype = ActiveEditorEntityManager->GetArchetypeForEntity(Entity);
+			const FMassArchetypeCompositionDescriptor& Composition = ActiveEditorEntityManager->GetArchetypeComposition(Archetype);
+
+			auto Callback = [&Composition](uint8_t ColumnIndex, TWeakObjectPtr<const UScriptStruct> ColumnType)
 			{
-				if (Column->IsChildOf(FMassFragment::StaticStruct()))
+				if (ColumnType.IsValid())
 				{
-					return Composition.Fragments.Contains(*Column);
+					if (ColumnType->IsChildOf(FMassFragment::StaticStruct()))
+					{
+						return Composition.Fragments.Contains(*ColumnType);
+					}
+					else if (ColumnType->IsChildOf(FMassTag::StaticStruct()))
+					{
+						return Composition.Tags.Contains(*ColumnType);
+					}
 				}
-				else if (Column->IsChildOf(FMassTag::StaticStruct()))
+				return false;
+			};
+			return Conditions.Verify(Callback);
+		}
+		else
+		{
+			auto Callback = [Row, &CommandBuffer = Environment->GetDirectDeferredCommands()]
+				(uint8_t ColumnIndex, TWeakObjectPtr<const UScriptStruct> ColumnType)
+			{
+				if (ColumnType.IsValid())
 				{
-					return Composition.Tags.Contains(*Column);
+					return CommandBuffer.HasColumn(Row, ColumnType.Get());
 				}
-			}
-			return false;
-		};
-		return Conditions.Verify(Callback);
+				return false;
+			};
+			return Conditions.Verify(Callback);
+		}
 	}
 	return false;
 }
diff --git a/Engine/Plugins/Experimental/TypedElementsDataStorage/Source/TypedElementsDataStorage/Private/TypedElementDatabase.h b/Engine/Plugins/Experimental/TypedElementsDataStorage/Source/TypedElementsDataStorage/Private/TypedElementDatabase.h
--- a/Engine/Plugins/Experimental/TypedElementsDataStorage/Source/TypedElementsDataStorage/Private/TypedElementDatabase.h
+++ b/Engine/Plugins/Experimental/TypedElementsDataStorage/Source/TypedElementsDataStorage/Private/TypedElementDatabase.h
@@ -88,30 +88,29 @@ public:
 	TypedElementRowHandle ReserveRow() override;
 	TypedElementRowHandle AddRow(TypedElementTableHandle Table) override;
 	bool AddRow(TypedElementRowHandle ReservedRow, TypedElementTableHandle Table) override;
-	bool BatchAddRow(TypedElementTableHandle Table, int32 Count, TypedElementDataStorageCreationCallbackRef OnCreated) override;
+	bool BatchAddRow(TypedElementTableHandle Table, int32 Count, TypedElementDataStorage::RowCreationCallbackRef OnCreated) override;
 	bool BatchAddRow(TypedElementTableHandle Table, TConstArrayView<TypedElementRowHandle> ReservedHandles,
-		TypedElementDataStorageCreationCallbackRef OnCreated) override;
+		TypedElementDataStorage::RowCreationCallbackRef OnCreated) override;
 	void RemoveRow(TypedElementRowHandle Row) override;
 	bool IsRowAvailable(TypedElementRowHandle Row) const override;
 	bool HasRowBeenAssigned(TypedElementRowHandle Row) const override;
 
-	bool AddColumn(TypedElementRowHandle Row, const UScriptStruct* ColumnType) override;
-	bool AddColumn(TypedElementRowHandle Row, FTopLevelAssetPath ColumnName) override;
+	void AddColumn(TypedElementRowHandle Row, const UScriptStruct* ColumnType) override;
+	void AddColumnData(TypedElementRowHandle Row, const UScriptStruct* ColumnType,
+		const TypedElementDataStorage::ColumnCreationCallbackRef& Initializer,
+		TypedElementDataStorage::ColumnCopyOrMoveCallback Relocator) override;
 	void RemoveColumn(TypedElementRowHandle Row, const UScriptStruct* ColumnType) override;
-	void RemoveColumn(TypedElementRowHandle Row, FTopLevelAssetPath ColumnName) override;
-	void* AddOrGetColumnData(TypedElementRowHandle Row, const UScriptStruct* ColumnType) override;
-	ColumnDataResult AddOrGetColumnData(TypedElementRowHandle Row, FTopLevelAssetPath ColumnName) override;
+	void* AddOrGetColumnData(TypedElementRowHandle Row, const UScriptStruct* ColumnType,
+		const TypedElementDataStorage::ColumnCreationCallbackRef& Initializer,
+		TypedElementDataStorage::ColumnCopyOrMoveCallback Relocator) override;
 	void* GetColumnData(TypedElementRowHandle Row, const UScriptStruct* ColumnType) override;
 	const void* GetColumnData(TypedElementRowHandle Row, const UScriptStruct* ColumnType) const override;
-	ColumnDataResult AddOrGetColumnData(TypedElementRowHandle Row, FTopLevelAssetPath ColumnName,
-		TConstArrayView<TypedElement::ColumnUtils::Argument> Arguments) override;
-	bool AddColumns(TypedElementRowHandle Row, TConstArrayView<const UScriptStruct*> Columns) override;
+	void AddColumns(TypedElementRowHandle Row, TConstArrayView<const UScriptStruct*> Columns) override;
 	void RemoveColumns(TypedElementRowHandle Row, TConstArrayView<const UScriptStruct*> Columns) override;
-	bool AddRemoveColumns(TypedElementRowHandle Row, TConstArrayView<const UScriptStruct*> ColumnsToAdd,
+	void AddRemoveColumns(TypedElementRowHandle Row, TConstArrayView<const UScriptStruct*> ColumnsToAdd,
 		TConstArrayView<const UScriptStruct*> ColumnsToRemove) override;
-	bool BatchAddRemoveColumns(TConstArrayView<TypedElementRowHandle> Rows,TConstArrayView<const UScriptStruct*> ColumnsToAdd,
+	void BatchAddRemoveColumns(TConstArrayView<TypedElementRowHandle> Rows,TConstArrayView<const UScriptStruct*> ColumnsToAdd,
 		TConstArrayView<const UScriptStruct*> ColumnsToRemove) override;
-	ColumnDataResult GetColumnData(TypedElementRowHandle Row, FTopLevelAssetPath ColumnName) override;
 	bool HasColumns(TypedElementRowHandle Row, TConstArrayView<const UScriptStruct*> ColumnTypes) const override;
 	bool HasColumns(TypedElementRowHandle Row, TConstArrayView<TWeakObjectPtr<const UScriptStruct>> ColumnTypes) const override;
 	bool MatchesColumns(TypedElementDataStorage::RowHandle Row, const TypedElementDataStorage::FQueryConditions& Conditions) const override;
@@ -170,8 +169,6 @@ private:
 
 	TUniquePtr<FTypedElementDatabaseEnvironment> Environment;
 	
-	FTypedElementDatabaseCommandBuffer::CommandBuffer DeferredCommands;
-	
 	FTypedElementOnDataStorageUpdate OnUpdateDelegate;
 	FDelegateHandle OnPreMassTickHandle;
 	FDelegateHandle OnPostMassTickHandle;
diff --git a/Engine/Plugins/Experimental/TypedElementsDataStorage/Source/TypedElementsDataStorage/Private/TypedElementDatabaseCommandBuffer.cpp b/Engine/Plugins/Experimental/TypedElementsDataStorage/Source/TypedElementsDataStorage/Private/TypedElementDatabaseCommandBuffer.cpp
--- a/Engine/Plugins/Experimental/TypedElementsDataStorage/Source/TypedElementsDataStorage/Private/TypedElementDatabaseCommandBuffer.cpp
+++ b/Engine/Plugins/Experimental/TypedElementsDataStorage/Source/TypedElementsDataStorage/Private/TypedElementDatabaseCommandBuffer.cpp
@@ -2,55 +2,162 @@
 
 #include "TypedElementDatabaseCommandBuffer.h"
 
+#include "HAL/UnrealMemory.h"
 #include "MassEntityManager.h"
 #include "TypedElementDatabaseEnvironment.h"
 
+// 
+// Commands section
+//
+
+FTypedElementDatabaseCommandBuffer::FTypedElementDatabaseCommandBuffer(FTypedElementDatabaseEnvironment& InEnvironment)
+	: Environment(InEnvironment)
+{
+}
+
+void* FTypedElementDatabaseCommandBuffer::GetQueuedDataColumn(
+	TypedElementDataStorage::RowHandle Row, const UScriptStruct* ColumnType)
+{
+	checkf(ColumnType->IsChildOf(FMassFragment::StaticStruct()),
+		TEXT("Trying to get the column '%s' which isn't a data column."), *ColumnType->GetName());
+
+	void* const* StoredData = PendingColumns.Find(PendingColumnMappingKey(Row, ColumnType));
+	if (StoredData)
+	{
+		if (*StoredData != nullptr)
+		{
+			return *StoredData;
+		}
+		else
+		{
+			// If the column was created but had no data assigned to it, create data for it now. If this code path triggers
+			// a lot there may be a large number of AddColumn followed by GetColumn calls. These can be more efficiently done
+			// with an AddorGetColumn call.
+			void* Result = Queue_AddDataColumnCommandUnitialized(Row, ColumnType,
+				[](const UScriptStruct& ColumnType, void* Destination, void* Source)
+				{
+					ColumnType.CopyScriptStruct(Destination, Source);
+				});
+			return Result;
+		}
+	}
+	else
+	{
+		return nullptr;
+	}
+}
+
+bool FTypedElementDatabaseCommandBuffer::HasColumn(TypedElementDataStorage::RowHandle Row, const UScriptStruct* ColumnType) const
+{
+	return PendingColumns.Contains(PendingColumnMappingKey(Row, ColumnType));
+}
+
+void FTypedElementDatabaseCommandBuffer::Clear(TypedElementDataStorage::RowHandle Row)
+{
+	for (TMap<PendingColumnMappingKey, void*>::TIterator It = PendingColumns.CreateIterator(); It; ++It)
+	{
+		if (It.Key().Key == Row)
+		{
+			It.RemoveCurrent();
+		}
+	}
+
+	for (TArray<FCommand>::TIterator It = Commands.CreateIterator(); It; ++It)
+	{
+		if (It->Row == Row)
+		{
+			It.RemoveCurrent(); // Don't swap as the order of commands is important.
+		}
+	}
+}
+
+FTypedElementDatabaseCommandBuffer::FAddDataColumnCommand::~FAddDataColumnCommand()
+{
+	if (ColumnType.IsValid() && (ColumnType->StructFlags & (STRUCT_IsPlainOldData | STRUCT_NoDestructor)) == 0)
+	{
+		ColumnType->DestroyStruct(Data);
+	}
+}
+
 //
 // Queue section
 //
 
-void FTypedElementDatabaseCommandBuffer::Queue_AddColumnCommand(
-	CommandBuffer& Buffer,
-	TypedElementDataStorage::RowHandle Row,
-	const UScriptStruct* ColumnType)
+void FTypedElementDatabaseCommandBuffer::Queue_AddColumnCommand(TypedElementDataStorage::RowHandle Row, const UScriptStruct* ColumnType)
 {
-	AddCommand(Buffer, Row, FAddColumnCommand
+	AddCommand(Row, FAddColumnCommand
 		{ 
 			.ColumnType = ColumnType 
 		});
+	PendingColumns.FindOrAdd(PendingColumnMappingKey(Row, ColumnType), nullptr);
 }
 
-void FTypedElementDatabaseCommandBuffer::Queue_AddColumnsCommand(
-	CommandBuffer& Buffer,
-	TypedElementDataStorage::RowHandle Row,
-	FMassFragmentBitSet FragmentsToAdd,
-	FMassTagBitSet TagsToAdd)
+void* FTypedElementDatabaseCommandBuffer::Queue_AddDataColumnCommandUnitialized(
+	TypedElementDataStorage::RowHandle Row, const UScriptStruct* ColumnType, TypedElementDataStorage::ColumnCopyOrMoveCallback Relocator)
 {
-	AddCommand(Buffer, Row, FAddColumnsCommand
+	checkf(ColumnType->IsChildOf(FMassFragment::StaticStruct()),
+		TEXT("Trying to queue a data column creation for '%s' which isn't a data column."), *ColumnType->GetName());
+
+	PendingColumnMappingKey Key(Row, ColumnType);
+	if (void** StoredData = PendingColumns.Find(Key); StoredData == nullptr || *StoredData == nullptr)
+	{
+		void* Data = Environment.GetScratchBuffer().Allocate(ColumnType->GetStructureSize(), ColumnType->GetMinAlignment());
+		// Initialize to zero to replicate the default from Mass.
+		FMemory::Memzero(Data, ColumnType->GetStructureSize());
+		PendingColumns.Add(PendingColumnMappingKey(Row, ColumnType), Data);
+		AddCommand(Row, FAddDataColumnCommand
+			{
+				.ColumnType = ColumnType,
+				.Relocator = Relocator,
+				.Data = Data
+			});
+		return Data;
+	}
+	else
+	{
+		return *StoredData;
+	}
+}
+
+void FTypedElementDatabaseCommandBuffer::Queue_AddColumnsCommand(TypedElementDataStorage::RowHandle Row, 
+	FMassFragmentBitSet FragmentsToAdd, FMassTagBitSet TagsToAdd)
+{
+	auto AddColumn = [this, Row](const UScriptStruct* ColumnType)
+	{
+		PendingColumns.FindOrAdd(PendingColumnMappingKey(Row, ColumnType), nullptr);
+		return true;
+	};
+	FragmentsToAdd.ExportTypes(AddColumn);
+	TagsToAdd.ExportTypes(AddColumn);
+
+	AddCommand(Row, FAddColumnsCommand
 		{ 
 			.FragmentsToAdd = MoveTemp(FragmentsToAdd), 
 			.TagsToAdd = MoveTemp(TagsToAdd) 
 		});
 }
 
-void FTypedElementDatabaseCommandBuffer::Queue_RemoveColumnCommand(
-	CommandBuffer& Buffer,
-	TypedElementDataStorage::RowHandle Row,
-	const UScriptStruct* ColumnType)
+void FTypedElementDatabaseCommandBuffer::Queue_RemoveColumnCommand(TypedElementDataStorage::RowHandle Row, const UScriptStruct* ColumnType)
 {
-	AddCommand(Buffer, Row, FRemoveColumnCommand
+	AddCommand(Row, FRemoveColumnCommand
 		{ 
 			.ColumnType = ColumnType 
 		});
+	PendingColumns.Remove(PendingColumnMappingKey(Row, ColumnType));
 }
 
-void FTypedElementDatabaseCommandBuffer::Queue_RemoveColumnsCommand(
-	CommandBuffer& Buffer,
-	TypedElementDataStorage::RowHandle Row,
-	FMassFragmentBitSet FragmentsToRemove,
-	FMassTagBitSet TagsToRemove)
+void FTypedElementDatabaseCommandBuffer::Queue_RemoveColumnsCommand(TypedElementDataStorage::RowHandle Row,
+	FMassFragmentBitSet FragmentsToRemove, FMassTagBitSet TagsToRemove)
 {
-	AddCommand(Buffer, Row, FRemoveColumnsCommand
+	auto RemoveColumn = [this, Row](const UScriptStruct* ColumnType)
+	{
+		PendingColumns.Remove(PendingColumnMappingKey(Row, ColumnType));
+		return true;
+	};
+	FragmentsToRemove.ExportTypes(RemoveColumn);
+	TagsToRemove.ExportTypes(RemoveColumn);
+	
+	AddCommand(Row, FRemoveColumnsCommand
 		{ 
 			.FragmentsToRemove = MoveTemp(FragmentsToRemove), 
 			.TagsToRemove = MoveTemp(TagsToRemove) 
@@ -95,6 +202,33 @@ void FTypedElementDatabaseCommandBuffer::Execute_AddColumnCommand(
 	}
 }
 
+void FTypedElementDatabaseCommandBuffer::Execute_AddDataColumnCommand(
+	FMassEntityManager& MassEntityManager, TypedElementDataStorage::RowHandle Row, const UScriptStruct* ColumnType,
+	void* Data, TypedElementDataStorage::ColumnCopyOrMoveCallback Relocator)
+{
+	if (ColumnType)
+	{
+		checkf(ColumnType->IsChildOf(FMassFragment::StaticStruct()),
+			TEXT("Trying to create a data column for '%s' from a deferred command that isn't a data column."), *ColumnType->GetName());
+
+		FMassEntityHandle Entity = FMassEntityHandle::FromNumber(Row); 
+		FStructView Column = MassEntityManager.GetFragmentDataStruct(Entity, ColumnType);
+		// Only add if not already added to avoid asserts from Mass.
+		if (!Column.IsValid())
+		{
+			MassEntityManager.AddFragmentToEntity(Entity, ColumnType,
+				[Data, Relocator](void* Fragment, const UScriptStruct& FragmentType)
+				{
+					Relocator(FragmentType, Fragment, Data);
+				});
+		}
+		else
+		{
+			Relocator(*ColumnType, Column.GetMemory(), Data);
+		}
+	}
+}
+
 void FTypedElementDatabaseCommandBuffer::Execute_AddColumnsCommand(
 	FMassEntityManager& MassEntityManager,
 	TypedElementDataStorage::RowHandle Row,
@@ -136,11 +270,11 @@ void FTypedElementDatabaseCommandBuffer::Execute_RemoveColumnsCommand(
 	MassEntityManager.RemoveCompositionFromEntity(FMassEntityHandle::FromNumber(Row), RemoveComposition);
 }
 
-void FTypedElementDatabaseCommandBuffer::ProcessBuffer(CommandBuffer& Buffer, FTypedElementDatabaseEnvironment& Environment)
+void FTypedElementDatabaseCommandBuffer::ProcessCommands()
 {
 	using namespace TypedElementDataStorage;
 
-	Buffer.StableSort(
+	Commands.StableSort(
 		[](const FCommand& Lhs, const FCommand& Rhs)
 		{
 			return Lhs.Row < Rhs.Row;
@@ -151,13 +285,14 @@ void FTypedElementDatabaseCommandBuffer::ProcessBuffer(CommandBuffer& Buffer, FT
 		FMassEntityManager& EntityManager;
 		RowHandle Row;
 		void operator()(const FAddColumnCommand& Command) { Execute_AddColumnCommand(EntityManager, Row, Command.ColumnType.Get()); }
+		void operator()(const FAddDataColumnCommand& Command) { Execute_AddDataColumnCommand(EntityManager, Row, Command.ColumnType.Get(), Command.Data, Command.Relocator); }
 		void operator()(const FAddColumnsCommand& Command) { Execute_AddColumnsCommand(EntityManager, Row, Command.FragmentsToAdd, Command.TagsToAdd); }
 		void operator()(const FRemoveColumnCommand& Command) { Execute_RemoveColumnCommand(EntityManager, Row, Command.ColumnType.Get()); }
 		void operator()(const FRemoveColumnsCommand& Command) { Execute_RemoveColumnsCommand(EntityManager, Row, Command.FragmentsToRemove, Command.TagsToRemove); }
 	};
 	
 	FProcessor Processor{ .EntityManager = Environment.GetMassEntityManager() };
-	for (FCommand& Command : Buffer)
+	for (FCommand& Command : Commands)
 	{
 		if (Execute_HasRowBeenAssigned(Environment.GetMassEntityManager(), Command.Row))
 		{
@@ -165,6 +300,9 @@ void FTypedElementDatabaseCommandBuffer::ProcessBuffer(CommandBuffer& Buffer, FT
 			Visit(Processor, Command.Data);
 		}
 	}
+
+	PendingColumns.Reset();
+	Commands.Reset();
 }
 
 // 
@@ -172,10 +310,10 @@ void FTypedElementDatabaseCommandBuffer::ProcessBuffer(CommandBuffer& Buffer, FT
 //
 
 template<typename T>
-void FTypedElementDatabaseCommandBuffer::AddCommand(CommandBuffer& Buffer, TypedElementDataStorage::RowHandle Row, T&& Args)
+void FTypedElementDatabaseCommandBuffer::AddCommand(TypedElementDataStorage::RowHandle Row, T&& Args)
 {
 	FCommand Command;
 	Command.Row = Row;
 	Command.Data.Emplace<T>(Forward<T>(Args));
-	Buffer.Add(MoveTemp(Command));
+	Commands.Add(MoveTemp(Command));
 }
diff --git a/Engine/Plugins/Experimental/TypedElementsDataStorage/Source/TypedElementsDataStorage/Private/TypedElementDatabaseCommandBuffer.h b/Engine/Plugins/Experimental/TypedElementsDataStorage/Source/TypedElementsDataStorage/Private/TypedElementDatabaseCommandBuffer.h
--- a/Engine/Plugins/Experimental/TypedElementsDataStorage/Source/TypedElementsDataStorage/Private/TypedElementDatabaseCommandBuffer.h
+++ b/Engine/Plugins/Experimental/TypedElementsDataStorage/Source/TypedElementsDataStorage/Private/TypedElementDatabaseCommandBuffer.h
@@ -2,6 +2,8 @@
 
 #pragma once
 
+#include <type_traits>
+#include "Elements/Common/TypedElementCommonTypes.h"
 #include "Elements/Common/TypedElementHandles.h"
 #include "MassArchetypeTypes.h"
 #include "Misc/TVariant.h"
@@ -14,11 +16,64 @@ struct FMassProcessingPhaseManager;
 
 class FTypedElementDatabaseCommandBuffer final
 {
+public:
+	explicit FTypedElementDatabaseCommandBuffer(FTypedElementDatabaseEnvironment& InEnvironment);
+
+	/** Returns a pointer to the queued data column, if it exists and hasn't been processed yet. Otherwise null is returned. */
+	void* GetQueuedDataColumn(TypedElementDataStorage::RowHandle Row, const UScriptStruct* ColumnType);
+	/** Returns if the column on the provided row is pending processing. */
+	bool HasColumn(TypedElementDataStorage::RowHandle Row, const UScriptStruct* ColumnType) const;
+	/** Clears all entries in the internal cache for the provided row. */
+	void Clear(TypedElementDataStorage::RowHandle Row);
+	
+	/**
+	 * @section Queue_* functions create a command that can be stored for later execution in the provided command buffer.
+	 * Typically the Execute_* counter part will be called for execution.
+	 */
+	void Queue_AddColumnCommand(TypedElementDataStorage::RowHandle Row, const UScriptStruct* ColumnType);
+	void* Queue_AddDataColumnCommandUnitialized(TypedElementDataStorage::RowHandle Row, const UScriptStruct* ColumnType, 
+		TypedElementDataStorage::ColumnCopyOrMoveCallback Relocator);
+	void Queue_AddColumnsCommand(TypedElementDataStorage::RowHandle Row, FMassFragmentBitSet FragmentsToAdd, FMassTagBitSet TagsToAdd);
+
+	void Queue_RemoveColumnCommand(TypedElementDataStorage::RowHandle Row, const UScriptStruct* ColumnType);
+	void Queue_RemoveColumnsCommand(TypedElementDataStorage::RowHandle Row, FMassFragmentBitSet FragmentsToRemove, FMassTagBitSet TagsToRemove);
+
+	/**
+	 * @section Execute_* functions directly execute a command with limited validation checks.
+	 */
+	
+	static bool Execute_IsRowAvailable(const FMassEntityManager& MassEntityManager, TypedElementDataStorage::RowHandle Row);
+	static bool Execute_HasRowBeenAssigned(const FMassEntityManager& MassEntityManager, TypedElementDataStorage::RowHandle Row);
+
+	static void Execute_AddColumnCommand(
+		FMassEntityManager& MassEntityManager, TypedElementDataStorage::RowHandle Row, const UScriptStruct* ColumnType);
+	static void Execute_AddDataColumnCommand(
+		FMassEntityManager& MassEntityManager, TypedElementDataStorage::RowHandle Row, const UScriptStruct* ColumnType,
+		void* Data, TypedElementDataStorage::ColumnCopyOrMoveCallback Relocator);
+	static void Execute_AddColumnsCommand(
+		FMassEntityManager& MassEntityManager, TypedElementDataStorage::RowHandle Row,
+		FMassFragmentBitSet FragmentsToAdd, FMassTagBitSet TagsToAdd);
+	static void Execute_RemoveColumnCommand(
+		FMassEntityManager& MassEntityManager, TypedElementDataStorage::RowHandle Row, const UScriptStruct* ColumnType);
+	static void Execute_RemoveColumnsCommand(
+		FMassEntityManager& MassEntityManager, TypedElementDataStorage::RowHandle Row,
+		FMassFragmentBitSet FragmentsToRemove, FMassTagBitSet TagsToRemove);
+
+	void ProcessCommands();
+
 private:
 	struct FAddColumnCommand
 	{
 		TWeakObjectPtr<const UScriptStruct> ColumnType;
 	};
+	struct FAddDataColumnCommand
+	{
+		TWeakObjectPtr<const UScriptStruct> ColumnType;
+		TypedElementDataStorage::ColumnCopyOrMoveCallback Relocator;
+		void* Data;
+
+		~FAddDataColumnCommand();
+	};
 	struct FAddColumnsCommand
 	{
 		FMassFragmentBitSet FragmentsToAdd;
@@ -33,7 +88,12 @@ private:
 		FMassFragmentBitSet FragmentsToRemove;
 		FMassTagBitSet TagsToRemove;
 	};
-	using CommandData = TVariant<FAddColumnCommand, FAddColumnsCommand, FRemoveColumnCommand, FRemoveColumnsCommand>;
+	using CommandData = TVariant<
+		FAddColumnCommand,
+		FAddDataColumnCommand,
+		FAddColumnsCommand,
+		FRemoveColumnCommand,
+		FRemoveColumnsCommand>;
 
 	struct FCommand
 	{
@@ -41,61 +101,11 @@ private:
 		CommandData Data;
 	};
 
-public:
-	using CommandBuffer = TArray<FCommand>;
-
-	/**
-	 * @section Queue_* functions create a command that can be stored for later execution in the provided command buffer.
-	 * Typically the Execute_* counter part will be called for execution.
-	 */
-	static void Queue_AddColumnCommand(
-		CommandBuffer& Buffer,
-		TypedElementDataStorage::RowHandle Row,
-		const UScriptStruct* ColumnType);
-	static void Queue_AddColumnsCommand(
-		CommandBuffer& Buffer,
-		TypedElementDataStorage::RowHandle Row,
-		FMassFragmentBitSet FragmentsToAdd,
-		FMassTagBitSet TagsToAdd);
-	static void Queue_RemoveColumnCommand(
-		CommandBuffer& Buffer,
-		TypedElementDataStorage::RowHandle Row,
-		const UScriptStruct* ColumnType);
-	static void Queue_RemoveColumnsCommand(
-		CommandBuffer& Buffer,
-		TypedElementDataStorage::RowHandle Row,
-		FMassFragmentBitSet FragmentsToRemove,
-		FMassTagBitSet TagsToRemove);
-
-	/**
-	 * @section Execute_* functions directly execute a command with limited validation checks.
-	 */
-	
-	static bool Execute_IsRowAvailable(const FMassEntityManager& MassEntityManager, TypedElementDataStorage::RowHandle Row);
-	static bool Execute_HasRowBeenAssigned(const FMassEntityManager& MassEntityManager, TypedElementDataStorage::RowHandle Row);
-
-	static void Execute_AddColumnCommand(
-		FMassEntityManager& MassEntityManager, 
-		TypedElementDataStorage::RowHandle Row,
-		const UScriptStruct* ColumnType);
-	static void Execute_AddColumnsCommand(
-		FMassEntityManager& MassEntityManager, 
-		TypedElementDataStorage::RowHandle Row,
-		FMassFragmentBitSet FragmentsToAdd, 
-		FMassTagBitSet TagsToAdd);
-	static void Execute_RemoveColumnCommand(
-		FMassEntityManager& MassEntityManager, 
-		TypedElementDataStorage::RowHandle Row,
-		const UScriptStruct* ColumnType);
-	static void Execute_RemoveColumnsCommand(
-		FMassEntityManager& MassEntityManager, 
-		TypedElementDataStorage::RowHandle Row,
-		FMassFragmentBitSet FragmentsToRemove, 
-		FMassTagBitSet TagsToRemove);
-
-	static void ProcessBuffer(CommandBuffer& Buffer, FTypedElementDatabaseEnvironment& Environment);
-
-private:
 	template<typename T>
-	static void AddCommand(CommandBuffer& Buffer, TypedElementDataStorage::RowHandle Row, T&& Args);
-};
\ No newline at end of file
+	void AddCommand(TypedElementDataStorage::RowHandle Row, T&& Args);
+
+	using PendingColumnMappingKey = TPair<TypedElementDataStorage::RowHandle, TWeakObjectPtr<const UScriptStruct>>;
+	TMap<PendingColumnMappingKey, void*> PendingColumns;
+	TArray<FCommand> Commands;
+	FTypedElementDatabaseEnvironment& Environment;
+};
diff --git a/Engine/Plugins/Experimental/TypedElementsDataStorage/Source/TypedElementsDataStorage/Private/TypedElementDatabaseEnvironment.cpp b/Engine/Plugins/Experimental/TypedElementsDataStorage/Source/TypedElementsDataStorage/Private/TypedElementDatabaseEnvironment.cpp
--- a/Engine/Plugins/Experimental/TypedElementsDataStorage/Source/TypedElementsDataStorage/Private/TypedElementDatabaseEnvironment.cpp
+++ b/Engine/Plugins/Experimental/TypedElementsDataStorage/Source/TypedElementsDataStorage/Private/TypedElementDatabaseEnvironment.cpp
@@ -4,11 +4,22 @@
 
 FTypedElementDatabaseEnvironment::FTypedElementDatabaseEnvironment(
 	FMassEntityManager& InMassEntityManager, FMassProcessingPhaseManager& InMassPhaseManager)
-	: MassEntityManager(InMassEntityManager)
+	: DirectDeferredCommands(*this)
+	, MassEntityManager(InMassEntityManager)
 	, MassPhaseManager(InMassPhaseManager)
 {
 }
 
+FTypedElementDatabaseCommandBuffer& FTypedElementDatabaseEnvironment::GetDirectDeferredCommands()
+{
+	return DirectDeferredCommands;
+}
+
+const FTypedElementDatabaseCommandBuffer& FTypedElementDatabaseEnvironment::GetDirectDeferredCommands() const
+{
+	return DirectDeferredCommands;
+}
+
 FTypedElementDatabaseIndexTable& FTypedElementDatabaseEnvironment::GetIndexTable()
 {
 	return IndexTable;
diff --git a/Engine/Plugins/Experimental/TypedElementsDataStorage/Source/TypedElementsDataStorage/Private/TypedElementDatabaseEnvironment.h b/Engine/Plugins/Experimental/TypedElementsDataStorage/Source/TypedElementsDataStorage/Private/TypedElementDatabaseEnvironment.h
--- a/Engine/Plugins/Experimental/TypedElementsDataStorage/Source/TypedElementsDataStorage/Private/TypedElementDatabaseEnvironment.h
+++ b/Engine/Plugins/Experimental/TypedElementsDataStorage/Source/TypedElementsDataStorage/Private/TypedElementDatabaseEnvironment.h
@@ -4,6 +4,7 @@
 
 #include "MassEntityManager.h"
 #include "MassProcessingPhaseManager.h"
+#include "TypedElementDatabaseCommandBuffer.h"
 #include "TypedElementDatabaseScratchBuffer.h"
 #include "TypedElementDatabaseIndexTable.h"
 #include "Queries/TypedElementExtendedQueryStore.h"
@@ -13,6 +14,9 @@ class FTypedElementDatabaseEnvironment final
 public:
 	FTypedElementDatabaseEnvironment(FMassEntityManager& InMassEntityManager, FMassProcessingPhaseManager& InMassPhaseManager);
 
+	FTypedElementDatabaseCommandBuffer& GetDirectDeferredCommands();
+	const FTypedElementDatabaseCommandBuffer& GetDirectDeferredCommands() const;
+	
 	FTypedElementDatabaseIndexTable& GetIndexTable();
 	const FTypedElementDatabaseIndexTable& GetIndexTable() const;
 
@@ -32,6 +36,7 @@ public:
 	uint64 GetUpdateCycleId() const;
 
 private:
+	FTypedElementDatabaseCommandBuffer DirectDeferredCommands;
 	FTypedElementDatabaseIndexTable IndexTable;
 	FTypedElementDatabaseScratchBuffer ScratchBuffer;
 	FTypedElementExtendedQueryStore Queries;
diff --git a/Engine/Source/Runtime/TypedElementFramework/Public/Elements/Common/TypedElementCommonTypes.h b/Engine/Source/Runtime/TypedElementFramework/Public/Elements/Common/TypedElementCommonTypes.h
new file mode 100644
--- /dev/null
+++ b/Engine/Source/Runtime/TypedElementFramework/Public/Elements/Common/TypedElementCommonTypes.h
@@ -0,0 +1,67 @@
+// Copyright Epic Games, Inc. All Rights Reserved.
+
+#pragma once
+
+#include "CoreMinimal.h"
+#include "Elements/Common/TypedElementHandles.h"
+
+#include "TypedElementCommonTypes.generated.h"
+
+/**
+ * Base for the data structures for a column.
+ */
+USTRUCT()
+struct FTypedElementDataStorageColumn
+{
+	GENERATED_BODY()
+};
+
+/**
+ * Base for the data structures that act as tags to rows. Tags should not have any data.
+ */
+USTRUCT()
+struct FTypedElementDataStorageTag
+{
+	GENERATED_BODY()
+};
+
+namespace UE
+{
+	// Work around missing header/implementations on some platforms
+	namespace detail
+	{
+		template<typename T, typename U>
+		concept SameHelper = std::is_same_v<T, U>;
+	}
+	template<typename T, typename U>
+	concept same_as = detail::SameHelper<T, U> && detail::SameHelper<U, T>;
+
+	template<typename From, typename To>
+	concept convertible_to = std::is_convertible_v<From, To> && requires { static_cast<To>(std::declval<From>()); };
+
+	template<typename Derived, typename Base>
+	concept derived_from = std::is_base_of_v<Base, Derived> && std::is_convertible_v<const volatile Derived*, const volatile Base*>;
+}
+
+namespace TypedElementDataStorage
+{
+	// Standard callbacks.
+
+	using RowCreationCallbackRef = TFunctionRef<void(TypedElementDataStorage::RowHandle Row)>;
+	using ColumnCreationCallbackRef = TFunctionRef<void(void* Column, const UScriptStruct& ColumnType)>;
+	using ColumnCopyOrMoveCallback = void (*)(const UScriptStruct& ColumnType, void* Destination, void* Source);
+
+
+
+	// Template concepts to enforce type correctness.
+
+	template<typename T>
+	concept TDataColumnType = UE::derived_from<T, FTypedElementDataStorageColumn>;
+
+	template<typename T>
+	concept TTagColumnType = UE::derived_from<T, FTypedElementDataStorageTag>;
+
+	template<typename T>
+	concept TColumnType = TDataColumnType<T> || TTagColumnType<T>;
+
+} // namespace TypedElementDataStorage
\ No newline at end of file
diff --git a/Engine/Source/Runtime/TypedElementFramework/Public/Elements/Framework/TypedElementQueryBuilder.h b/Engine/Source/Runtime/TypedElementFramework/Public/Elements/Framework/TypedElementQueryBuilder.h
--- a/Engine/Source/Runtime/TypedElementFramework/Public/Elements/Framework/TypedElementQueryBuilder.h
+++ b/Engine/Source/Runtime/TypedElementFramework/Public/Elements/Framework/TypedElementQueryBuilder.h
@@ -3,6 +3,7 @@
 #pragma once
 
 #include "CoreMinimal.h"
+#include "Elements/Common/TypedElementCommonTypes.h"
 #include "Elements/Interfaces/TypedElementDataStorageInterface.h"
 #include "Templates/SubclassOf.h"
 #include "UObject/WeakObjectPtrTemplates.h"
@@ -10,23 +11,6 @@
 class UScriptStruct;
 class USubsystem;
 
-namespace UE
-{
-	// Work around missing header/implementations on some platforms
-	namespace detail
-	{
-		template<typename T, typename U>
-		concept SameHelper = std::is_same_v<T, U>;
-	}
-	template<typename T, typename U>
-	concept same_as = detail::SameHelper<T, U> && detail::SameHelper<U, T>;
-
-	template<typename From, typename To>
-	concept convertible_to = std::is_convertible_v<From, To> && requires { static_cast<To>(std::declval<From>()); };
-
-	template<typename Derived, typename Base>
-	concept derived_from = std::is_base_of_v<Base, Derived> && std::is_convertible_v<const volatile Derived*, const volatile Base*>;
-}
 /**
  * The TypedElementQueryBuilder allows for the construction of queries for use by the Typed Element Data Storage.
  * There are two types of queries, simple and normal. Simple queries are guaranteed to be supported by the data
@@ -127,15 +111,6 @@ namespace TypedElementQueryBuilder
 	TYPEDELEMENTFRAMEWORK_API const UScriptStruct* operator""_Type(const char* Name, std::size_t NameSize);
 	TYPEDELEMENTFRAMEWORK_API const UScriptStruct* operator""_TypeOptional(const char* Name, std::size_t NameSize);
 
-	template<typename T>
-	concept TDataColumnType = UE::derived_from<T, FTypedElementDataStorageColumn>;
-
-	template<typename T>
-	concept TTagColumnType = UE::derived_from<T, FTypedElementDataStorageTag>;
-
-	template<typename T>
-	concept TColumnType = TDataColumnType<T> || TTagColumnType<T>;
-
 	enum class EOptional
 	{
 		No,
@@ -177,15 +152,15 @@ namespace TypedElementQueryBuilder
 		TYPEDELEMENTFRAMEWORK_API FDependency DependsOn();
 		TYPEDELEMENTFRAMEWORK_API ITypedElementDataStorageInterface::FQueryDescription&& Compile();
 
-		template<TColumnType... TargetTypes>
+		template<TypedElementDataStorage::TColumnType... TargetTypes>
 		FSimpleQuery& All();
 		TYPEDELEMENTFRAMEWORK_API FSimpleQuery& All(const UScriptStruct* Target);
 		TYPEDELEMENTFRAMEWORK_API FSimpleQuery& All(TConstArrayView<const UScriptStruct*> Targets);
-		template<TColumnType... TargetTypes>
+		template<TypedElementDataStorage::TColumnType... TargetTypes>
 		FSimpleQuery& Any();
 		TYPEDELEMENTFRAMEWORK_API FSimpleQuery& Any(const UScriptStruct* Target);
 		TYPEDELEMENTFRAMEWORK_API FSimpleQuery& Any(TConstArrayView<const UScriptStruct*> Targets);
-		template<TColumnType... TargetTypes>
+		template<TypedElementDataStorage::TColumnType... TargetTypes>
 		FSimpleQuery& None();
 		TYPEDELEMENTFRAMEWORK_API FSimpleQuery& None(const UScriptStruct* Target);
 		TYPEDELEMENTFRAMEWORK_API FSimpleQuery& None(TConstArrayView<const UScriptStruct*> Targets);
@@ -226,14 +201,14 @@ namespace TypedElementQueryBuilder
 
 		TYPEDELEMENTFRAMEWORK_API FObserver(EEvent MonitorForEvent, const UScriptStruct* MonitoredColumn);
 
-		template<TColumnType ColumnType>
+		template<TypedElementDataStorage::TColumnType ColumnType>
 		static FObserver OnAdd();
-		template<TColumnType ColumnType>
+		template<TypedElementDataStorage::TColumnType ColumnType>
 		static FObserver OnRemove();
 
 		TYPEDELEMENTFRAMEWORK_API FObserver& SetEvent(EEvent MonitorForEvent);
 		TYPEDELEMENTFRAMEWORK_API FObserver& SetMonitoredColumn(const UScriptStruct* MonitoredColumn);
-		template<TColumnType ColumnType>
+		template<TypedElementDataStorage::TColumnType ColumnType>
 		FObserver& SetMonitoredColumn();
 		TYPEDELEMENTFRAMEWORK_API FObserver& ForceToGameThread(bool bForce);
 		TYPEDELEMENTFRAMEWORK_API FObserver& MakeActivatable(FName Name);
@@ -336,7 +311,7 @@ namespace TypedElementQueryBuilder
 		Select(FName Name, const CallbackType& Type, Class* Instance, Function&& Callback);
 
 		/** Request read-only access to the listed columns. */
-		template<TDataColumnType... TargetTypes>
+		template<TypedElementDataStorage::TDataColumnType... TargetTypes>
 		Select& ReadOnly();
 		/** Request read-only access to the listed columns. */
 		TYPEDELEMENTFRAMEWORK_API Select& ReadOnly(const UScriptStruct* Target);
@@ -346,7 +321,7 @@ namespace TypedElementQueryBuilder
 		 * Request read-only access to the listed columns. If optional is true read access will be given if the column is in the table but
 		 * it will not be used for finding matching tables. Columns bound with optional can not be bound to a query callback argument.
 		 */
-		template<TDataColumnType... TargetTypes>
+		template<TypedElementDataStorage::TDataColumnType... TargetTypes>
 		Select& ReadOnly(EOptional Optional);
 		/**
 		 * Request read-only access to the listed columns. If optional is true read access will be given if the column is in the table but
@@ -359,7 +334,7 @@ namespace TypedElementQueryBuilder
 		 */
 		TYPEDELEMENTFRAMEWORK_API Select& ReadOnly(TConstArrayView<const UScriptStruct*> Targets, EOptional Optionall);
 		/** Request read and write access to the listed columns. */
-		template<TDataColumnType... TargetTypes>
+		template<TypedElementDataStorage::TDataColumnType... TargetTypes>
 		Select& ReadWrite();
 		/** Request read and write access to the listed columns. */
 		TYPEDELEMENTFRAMEWORK_API Select& ReadWrite(const UScriptStruct* Target);
diff --git a/Engine/Source/Runtime/TypedElementFramework/Public/Elements/Framework/TypedElementQueryBuilder.inl b/Engine/Source/Runtime/TypedElementFramework/Public/Elements/Framework/TypedElementQueryBuilder.inl
--- a/Engine/Source/Runtime/TypedElementFramework/Public/Elements/Framework/TypedElementQueryBuilder.inl
+++ b/Engine/Source/Runtime/TypedElementFramework/Public/Elements/Framework/TypedElementQueryBuilder.inl
@@ -90,19 +90,19 @@ namespace TypedElementQueryBuilder
 	// FObserver
 	//
 	
-	template<TColumnType ColumnType>
+	template<TypedElementDataStorage::TColumnType ColumnType>
 	FObserver FObserver::OnAdd()
 	{
 		return FObserver(FObserver::EEvent::Add, ColumnType::StaticStruct());
 	}
 
-	template<TColumnType ColumnType>
+	template<TypedElementDataStorage::TColumnType ColumnType>
 	FObserver FObserver::OnRemove()
 	{
 		return FObserver(FObserver::EEvent::Remove, ColumnType::StaticStruct());
 	}
 
-	template<TColumnType ColumnType>
+	template<TypedElementDataStorage::TColumnType ColumnType>
 	FObserver& FObserver::SetMonitoredColumn()
 	{
 		return SetMonitoredColumn(ColumnType::StaticStruct());
@@ -1117,21 +1117,21 @@ e.g. void(FCachedQueryContext<Subsystem1, const Subsystem2>& Context, TypedEleme
 		Internal::BindQueryFunction<TypedElementDataStorage::IQueryContext, ValidateColumns>(Query.Callback.Function, Instance, Callback);
 	}
 
-	template<TDataColumnType... TargetTypes>
+	template<TypedElementDataStorage::TDataColumnType... TargetTypes>
 	Select& Select::ReadOnly()
 	{
 		ReadOnly({ TargetTypes::StaticStruct()... });
 		return *this;
 	}
 
-	template<TDataColumnType... TargetTypes>
+	template<TypedElementDataStorage::TDataColumnType... TargetTypes>
 	Select& Select::ReadOnly(EOptional Optional)
 	{
 		ReadOnly({ TargetTypes::StaticStruct()... }, Optional);
 		return *this;
 	}
 
-	template<TDataColumnType... TargetTypes>
+	template<TypedElementDataStorage::TDataColumnType... TargetTypes>
 	Select& Select::ReadWrite()
 	{
 		ReadWrite({ TargetTypes::StaticStruct()... });
@@ -1143,21 +1143,21 @@ e.g. void(FCachedQueryContext<Subsystem1, const Subsystem2>& Context, TypedEleme
 	// FSimpleQuery
 	//
 
-	template<TColumnType... TargetTypes>
+	template<TypedElementDataStorage::TColumnType... TargetTypes>
 	FSimpleQuery& FSimpleQuery::All()
 	{
 		All({ TargetTypes::StaticStruct()... });
 		return *this;
 	}
 
-	template<TColumnType... TargetTypes>
+	template<TypedElementDataStorage::TColumnType... TargetTypes>
 	FSimpleQuery& FSimpleQuery::Any()
 	{
 		Any({ TargetTypes::StaticStruct()... });
 		return *this;
 	}
 
-	template<TColumnType... TargetTypes>
+	template<TypedElementDataStorage::TColumnType... TargetTypes>
 	FSimpleQuery& FSimpleQuery::None()
 	{
 		None({ TargetTypes::StaticStruct()... });
diff --git a/Engine/Source/Runtime/TypedElementFramework/Public/Elements/Interfaces/TypedElementDataStorageInterface.h b/Engine/Source/Runtime/TypedElementFramework/Public/Elements/Interfaces/TypedElementDataStorageInterface.h
--- a/Engine/Source/Runtime/TypedElementFramework/Public/Elements/Interfaces/TypedElementDataStorageInterface.h
+++ b/Engine/Source/Runtime/TypedElementFramework/Public/Elements/Interfaces/TypedElementDataStorageInterface.h
@@ -2,9 +2,11 @@
 
 #pragma once
 
+#include <type_traits>
 #include "Containers/Array.h"
 #include "Containers/ContainersFwd.h"
 #include "Delegates/Delegate.h"
+#include "Elements/Common/TypedElementCommonTypes.h"
 #include "Elements/Common/TypedElementHandles.h"
 #include "Elements/Common/TypedElementQueryConditions.h"
 #include "Elements/Common/TypedElementQueryDescription.h"
@@ -25,14 +27,6 @@ class USubsystem;
 class UScriptStruct;
 class UTypedElementDataStorageFactory;
 
-struct ColumnDataResult
-{
-	/** Pointer to the structure that holds the description of the returned data. */
-	const UScriptStruct* Description;
-	/** Pointer to the column data. The type is guaranteed to match type described in Description. */
-	void* Data;
-};
-
 using TypedElementTableHandle = TypedElementDataStorage::TableHandle;
 static constexpr auto TypedElementInvalidTableHandle = TypedElementDataStorage::InvalidTableHandle;
 using TypedElementRowHandle = TypedElementDataStorage::RowHandle;
@@ -44,26 +38,6 @@ using FTypedElementOnDataStorageCreation = FSimpleMulticastDelegate;
 using FTypedElementOnDataStorageDestruction = FSimpleMulticastDelegate;
 using FTypedElementOnDataStorageUpdate = FSimpleMulticastDelegate;
 
-using TypedElementDataStorageCreationCallbackRef = TFunctionRef<void(TypedElementRowHandle Row)>;
-
-/**
- * Base for the data structures for a column.
- */
-USTRUCT()
-struct FTypedElementDataStorageColumn
-{
-	GENERATED_BODY()
-};
-
-/**
- * Base for the data structures that act as tags to rows. Tags should not have any data.
- */
-USTRUCT()
-struct FTypedElementDataStorageTag
-{
-	GENERATED_BODY()
-};
-
 UINTERFACE(MinimalAPI)
 class UTypedElementDataStorageInterface : public UInterface
 {
@@ -143,14 +117,14 @@ public:
 	 * Add multiple rows at once. For each new row the OnCreated callback is called. Callers are expected to use the callback to
 	 * initialize the row if needed.
 	 */
-	virtual bool BatchAddRow(TypedElementTableHandle Table, int32 Count, TypedElementDataStorageCreationCallbackRef OnCreated) = 0;
+	virtual bool BatchAddRow(TypedElementTableHandle Table, int32 Count, TypedElementDataStorage::RowCreationCallbackRef OnCreated) = 0;
 	/**
 	 * Add multiple rows at once. For each new row the OnCreated callback is called. Callers are expected to use the callback to
 	 * initialize the row if needed. This version uses a set of previously reserved rows. Any row that can't be used will be 
 	 * released.
 	 */
 	virtual bool BatchAddRow(TypedElementTableHandle Table, TConstArrayView<TypedElementRowHandle> ReservedHandles,
-		TypedElementDataStorageCreationCallbackRef OnCreated) = 0;
+		TypedElementDataStorage::RowCreationCallbackRef OnCreated) = 0;
 
 	/** Removes a previously reserved or added row. If the row handle is invalid or already removed, nothing happens */
 	virtual void RemoveRow(TypedElementRowHandle Row) = 0;
@@ -166,59 +140,87 @@ public:
 	 */
 
 	/** Adds a column to a row or does nothing if already added. */
-	virtual bool AddColumn(TypedElementRowHandle Row, const UScriptStruct* ColumnType) = 0;
-	virtual bool AddColumn(TypedElementRowHandle Row, FTopLevelAssetPath ColumnName) = 0;
+	virtual void AddColumn(TypedElementRowHandle Row, const UScriptStruct* ColumnType) = 0;
+	template<TypedElementDataStorage::TColumnType ColumnType>
+	void AddColumn(TypedElementRowHandle Row);
+	/**
+	 * Adds a new data column and initializes it. The relocator will be used to copy or move the column out of
+	 * its temporary location into the final table if the addition needs to be deferred.
+	 */
+	virtual void AddColumnData(TypedElementRowHandle Row, const UScriptStruct* ColumnType,
+		const TypedElementDataStorage::ColumnCreationCallbackRef& Initializer,
+		TypedElementDataStorage::ColumnCopyOrMoveCallback Relocator) = 0;
+	template<TypedElementDataStorage::TDataColumnType ColumnType>
+	void AddColumn(TypedElementRowHandle Row, ColumnType&& Column);
+
 	/**
 	 * Adds multiple columns from a row. This is typically more efficient than adding columns one 
 	 * at a time.
 	 */
-	virtual bool AddColumns(TypedElementRowHandle Row, TConstArrayView<const UScriptStruct*> Columns) = 0;
+	virtual void AddColumns(TypedElementRowHandle Row, TConstArrayView<const UScriptStruct*> Columns) = 0;
+	template<typename... Columns>
+	void AddColumns(TypedElementRowHandle Row);
 
 	/** Removes a column from a row or does nothing if already removed. */
 	virtual void RemoveColumn(TypedElementRowHandle Row, const UScriptStruct* ColumnType) = 0;
-	virtual void RemoveColumn(TypedElementRowHandle Row, FTopLevelAssetPath ColumnName) = 0;
+	template<typename Column>
+	void RemoveColumn(TypedElementRowHandle Row);
 	/**
 	 * Removes multiple columns from a row. This is typically more efficient than adding columns one
 	 * at a time.
 	 */
 	virtual void RemoveColumns(TypedElementRowHandle Row, TConstArrayView<const UScriptStruct*> Columns) = 0;
+	template<typename... Columns>
+	void RemoveColumns(TypedElementRowHandle Row);
+
 	/** 
 	 * Adds and removes the provided column types from the provided row. This is typically more efficient 
 	 * than individually adding and removing columns as well as being faster than adding and removing
 	 * columns separately.
 	 */
-	virtual bool AddRemoveColumns(TypedElementRowHandle Row, TConstArrayView<const UScriptStruct*> ColumnsToAdd,
+	virtual void AddRemoveColumns(TypedElementRowHandle Row, TConstArrayView<const UScriptStruct*> ColumnsToAdd,
 		TConstArrayView<const UScriptStruct*> ColumnsToRemove) = 0;
-
+	
 	/** Adds and removes the provided column types from the provided list of rows. */
-	virtual bool BatchAddRemoveColumns(
+	virtual void BatchAddRemoveColumns(
 		TConstArrayView<TypedElementRowHandle> Rows,
 		TConstArrayView<const UScriptStruct*> ColumnsToAdd,
 		TConstArrayView<const UScriptStruct*> ColumnsToRemove) = 0;
 
 	/**
-	 * Adds a new column to a row. If the column already exists it will be returned instead. If the colum couldn't
+	 * Adds a new column to a row. If the column already exists it will be returned instead. If the column couldn't
 	 * be added or the column type points to a tag an nullptr will be returned.
+	 * This version allows a new object to be initialized before OnAdd observers are called as well as provide a
+	 * bespoke call to move the object from temporary locations to the final table location.
 	 */
-	virtual void* AddOrGetColumnData(TypedElementRowHandle Row, const UScriptStruct* ColumnType) = 0;
-	virtual ColumnDataResult AddOrGetColumnData(TypedElementRowHandle Row, FTopLevelAssetPath ColumnName) = 0;
+	virtual void* AddOrGetColumnData(TypedElementRowHandle Row, const UScriptStruct* ColumnType,
+		const TypedElementDataStorage::ColumnCreationCallbackRef& Initializer,
+		TypedElementDataStorage::ColumnCopyOrMoveCallback Relocator) = 0;
 	/**
-	 * Sets the data of a column using the provided argument bag. This is only meant for simple initialization for 
-	 * fragments that use UPROPERTY to expose properties. For complex initialization or when the fragment type is known
-	 * it's recommended to use calls that work directly on the type for better performance and a wider range of configuration options.
-	 * If the column couldn't be created or the column name points to a tag, then the result will contain only nullptrs.
+	 * Returns a pointer to the column of the given row or creates a new one if not found.
+	 * Enables type deduction of ColumnType from Column argument.
+	 *
+	 * For example, FTransformColumn added and deduced from second argument:
+	 * StorageInterface->AddOrGetColumn(Row, FTransformColumn{.Transform = Transform});
 	 */
-	virtual ColumnDataResult AddOrGetColumnData(TypedElementRowHandle Row, FTopLevelAssetPath ColumnName,
-		TConstArrayView<TypedElement::ColumnUtils::Argument> Arguments) = 0;
+	template<typename ColumnType>
+	ColumnType* AddOrGetColumn(TypedElementRowHandle Row, ColumnType&& Column);
 	
 	/** Retrieves a pointer to the column of the given row or a nullptr if not found or if the column type is a tag. */
 	virtual void* GetColumnData(TypedElementRowHandle Row, const UScriptStruct* ColumnType) = 0;
 	virtual const void* GetColumnData(TypedElementRowHandle Row, const UScriptStruct* ColumnType) const = 0;
-	virtual ColumnDataResult GetColumnData(TypedElementRowHandle Row, FTopLevelAssetPath ColumnName) = 0;
-
+	/** Returns a pointer to the column of the given row or a nullptr if the type couldn't be found or the row doesn't exist. */
+	template<typename ColumnType>
+	ColumnType* GetColumn(TypedElementRowHandle Row);
+	template<typename ColumnType>
+	const ColumnType* GetColumn(TypedElementRowHandle Row) const;
+	
 	/** Determines if the provided row contains the collection of columns and tags. */
 	virtual bool HasColumns(TypedElementRowHandle Row, TConstArrayView<const UScriptStruct*> ColumnTypes) const = 0;
 	virtual bool HasColumns(TypedElementRowHandle Row, TConstArrayView<TWeakObjectPtr<const UScriptStruct>> ColumnTypes) const = 0;
+	template<typename... ColumnTypes>
+	bool HasColumns(TypedElementRowHandle Row) const;
+
 	/** Determines if the columns in the row match the query conditions. */
 	virtual bool MatchesColumns(TypedElementDataStorage::RowHandle Row, const TypedElementDataStorage::FQueryConditions& Conditions) const = 0;
 	
@@ -320,66 +322,6 @@ public:
 	/** Returns a pointer to the registered external system if found, otherwise null. */
 	virtual void* GetExternalSystemAddress(UClass* Target) = 0;
 
-	
-
-	/**
-	 *
-	 * The following are utility functions that are not part of the interface but are provided in order to make using the
-	 * interface easier.
-	 *
-	 */
-	
-	
-	/** Adds a column to a row or does nothing if already added. */
-	template<typename Column>
-	bool AddColumn(TypedElementRowHandle Row);
-
-	/** Removes a tag from a row or does nothing if already removed. */
-	template<typename Column>
-	void RemoveColumn(TypedElementRowHandle Row);
-	
-	/**
-	 * Adds multiple columns from a row. This is typically more efficient than adding columns one
-	 * at a time.
-	 */
-	template<typename... Columns>
-	void AddColumns(TypedElementRowHandle Row);
-
-	/**
-	 * Removes multiple columns from a row. This is typically more efficient than adding columns one
-	 * at a time.
-	 */
-	template<typename... Columns>
-	void RemoveColumns(TypedElementRowHandle Row);
-
-	/**
-	 * Returns a pointer to the column of the given row or creates a new one if not found. Optionally arguments can be provided
-	 * to update or initialize the column's data.
-	 */
-	template<typename ColumnType, typename... Args>
-	ColumnType* AddOrGetColumn(TypedElementRowHandle Row, Args... Arguments);
-
-	/**
-	 * Returns a pointer to the column of the given row or creates a new one if not found.
-	 * Enables type deduction of ColumnType from Column argument.
-	 * 
-	 * For example, FTransformColumn added and deduced from second argument:
-	 * StorageInterface->AddOrGetColumn(Row, FTransformColumn{.Transform = Transform});
-	 */
-	template<typename ColumnType>
-	ColumnType* AddOrGetColumn(TypedElementRowHandle Row, ColumnType&& Column);
-	
-	/** Returns a pointer to the column of the given row or a nullptr if the type couldn't be found or the row doesn't exist. */
-	template<typename ColumnType>
-	ColumnType* GetColumn(TypedElementRowHandle Row);
-
-	/** Returns a pointer to the column of the given row or a nullptr if the type couldn't be found or the row doesn't exist. */
-	template<typename ColumnType>
-	const ColumnType* GetColumn(TypedElementRowHandle Row) const;
-
-	template<typename... ColumnTypes>
-	bool HasColumns(TypedElementRowHandle Row) const;
-
 	/** Returns a pointer to the registered external system if found, otherwise null. */
 	template<typename SystemType>
 	SystemType* GetExternalSystem();
@@ -395,10 +337,10 @@ const FactoryT* ITypedElementDataStorageInterface::FindFactory() const
 	return static_cast<const FactoryT*>(FindFactory(FactoryT::StaticClass()));
 }
 
-template<typename Column>
-bool ITypedElementDataStorageInterface::AddColumn(TypedElementRowHandle Row)
+template<TypedElementDataStorage::TColumnType Column>
+void ITypedElementDataStorageInterface::AddColumn(TypedElementRowHandle Row)
 {
-	return AddColumn(Row, Column::StaticStruct());
+	AddColumn(Row, Column::StaticStruct());
 }
 
 template<typename Column>
@@ -419,46 +361,38 @@ void ITypedElementDataStorageInterface::RemoveColumns(TypedElementRowHandle Row)
 	RemoveColumns(Row, { Columns::StaticStruct()...});
 }
 
-template<typename ColumnType, typename FirstArg, typename... NextArgs>
-struct TConstructFromSlicedObjectDisabler
+template<TypedElementDataStorage::TDataColumnType ColumnType>
+void ITypedElementDataStorageInterface::AddColumn(TypedElementRowHandle Row, ColumnType&& Column)
 {
-	// Fail assertion if the only argument is derived from FTypedElementDataStorageColumn and not the same as ColumnType
-	// This gives a good indication that object slicing is probably happening.
-	static_assert(!(
-		sizeof...(NextArgs) == 0 &&                                    // There is only one argument to the callback and ...
-		std::is_base_of_v<FTypedElementDataStorageColumn, FirstArg> && // ... the type of the argument derives from FTypedElementDataStorageColumn
-		!std::is_same_v<ColumnType, FirstArg>),                        // ... but it isn't the same type as the column we are trying to add
-		"Probable object slicing detected. The invoked constructor of ColumnType uses a different column type as it's first argument. "
-		"This is detected as a likely object slice and disabled. "
-		"If this is what you intended, then explicitly slice the object using the constructor before passing it as an argument");
-};
-
-template<typename ColumnType, typename... Args>
-ColumnType* ITypedElementDataStorageInterface::AddOrGetColumn(TypedElementRowHandle Row, Args... Arguments)
-{
-	ColumnType* Result = static_cast<ColumnType*>(AddOrGetColumnData(Row, ColumnType::StaticStruct()));
-	if constexpr (sizeof...(Arguments) > 0)
-	{
-		[[maybe_unused]] TConstructFromSlicedObjectDisabler<ColumnType, Args...> SliceDisabler;
-		
-		if (Result)
-		{
-			new(Result) ColumnType{ std::forward<Args>(Arguments)... };
-		}
-	}
-	return Result;
+	AddOrGetColumn<ColumnType>(Row, Forward<ColumnType>(Column));
 }
 
 template<typename ColumnType>
 ColumnType* ITypedElementDataStorageInterface::AddOrGetColumn(TypedElementRowHandle Row, ColumnType&& Column)
 {
-	ColumnType* Result = static_cast<ColumnType*>(AddOrGetColumnData(Row, ColumnType::StaticStruct()));
-
-	if (Result)
-	{
-		new(Result) ColumnType(MoveTemp(Column));
-	}
-	return Result;
+	return reinterpret_cast<ColumnType*>(AddOrGetColumnData(Row, ColumnType::StaticStruct(),
+		[&Column](void* ColumnData, const UScriptStruct&)
+		{
+			if constexpr (std::is_move_constructible_v<ColumnType>)
+			{
+				new(ColumnData) ColumnType(MoveTemp(Column));
+			}
+			else
+			{
+				new(ColumnData) ColumnType(Column);
+			}
+		},
+		[](const UScriptStruct&, void* Destination, void* Source)
+		{
+			if constexpr (std::is_move_assignable_v<ColumnType>)
+			{
+				*reinterpret_cast<ColumnType*>(Destination) = MoveTemp(*reinterpret_cast<ColumnType*>(Source));
+			}
+			else
+			{
+				*reinterpret_cast<ColumnType*>(Destination) = *reinterpret_cast<ColumnType*>(Source);
+			}
+		}));
 }
 
 template<typename ColumnType>
