Removed `bSimpleQuery` from TEDS' queries.

Originally the idea was to have simple queries that map directly to Mass and complex queries have additional processing around them. This has not been implemented beyond the initial placeholders, which are currently causing invalid assert to trigger. Earlier change lists already removed the advanced query operators and introduced the query conditions which will be used for this in the future. This change removes the `bSimpleQuery` boolean, but will otherwise keep the SimpleQuery for now.

diff --git a/Engine/Plugins/Experimental/TypedElementsDataStorage/Source/TypedElementsDataStorage/Private/Queries/TypedElementExtendedQueryStore.cpp b/Engine/Plugins/Experimental/TypedElementsDataStorage/Source/TypedElementsDataStorage/Private/Queries/TypedElementExtendedQueryStore.cpp
--- a/Engine/Plugins/Experimental/TypedElementsDataStorage/Source/TypedElementsDataStorage/Private/Queries/TypedElementExtendedQueryStore.cpp
+++ b/Engine/Plugins/Experimental/TypedElementsDataStorage/Source/TypedElementsDataStorage/Private/Queries/TypedElementExtendedQueryStore.cpp
@@ -184,27 +184,20 @@ TypedElementDataStorage::FQueryResult FTypedElementExtendedQueryStore::RunQuery(
 
 	if (FTypedElementExtendedQuery* QueryData = Get(Query))
 	{
-		if (QueryData->Description.bSimpleQuery)
-		{
-			switch (QueryData->Description.Action)
-			{
-			case ActionType::None:
-				Result.Completed = CompletionType::Fully;
-				break;
-			case ActionType::Select:
-				// Fall through: There's nothing to callback to, so only return the total count.
-			case ActionType::Count:
-				Result.Count = QueryData->NativeQuery.GetNumMatchingEntities(EntityManager);
-				Result.Completed = CompletionType::Fully;
-				break;
-			default:
-				Result.Completed = CompletionType::Unsupported;
-				break;
-			}
-		}
-		else
+		switch (QueryData->Description.Action)
 		{
+		case ActionType::None:
+			Result.Completed = CompletionType::Fully;
+			break;
+		case ActionType::Select:
+			// Fall through: There's nothing to callback to, so only return the total count.
+		case ActionType::Count:
+			Result.Count = QueryData->NativeQuery.GetNumMatchingEntities(EntityManager);
+			Result.Completed = CompletionType::Fully;
+			break;
+		default:
 			Result.Completed = CompletionType::Unsupported;
+			break;
 		}
 	}
 	else
@@ -549,43 +542,40 @@ bool FTypedElementExtendedQueryStore::SetupConditions(
 		return true;
 	}
 
-	if (ensureMsgf(Query.bSimpleQuery, TEXT("The Data Storage back-end currently only supports simple queries.")))
+	if (ensureMsgf(Query.ConditionTypes.Num() == Query.ConditionOperators.Num(),
+		TEXT("The types and operators for a typed element query have gone out of sync.")))
 	{
-		if (ensureMsgf(Query.ConditionTypes.Num() == Query.ConditionOperators.Num(),
-			TEXT("The types and operators for a typed element query have gone out of sync.")))
+		const DSI::FQueryDescription::FOperator* Operand = Query.ConditionOperators.GetData();
+		for (DSI::FQueryDescription::EOperatorType Type : Query.ConditionTypes)
 		{
-			const DSI::FQueryDescription::FOperator* Operand = Query.ConditionOperators.GetData();
-			for (DSI::FQueryDescription::EOperatorType Type : Query.ConditionTypes)
+			EMassFragmentPresence Presence;
+			switch (Type)
 			{
-				EMassFragmentPresence Presence;
-				switch (Type)
-				{
-				case DSI::FQueryDescription::EOperatorType::SimpleAll:
-					Presence = EMassFragmentPresence::All;
-					break;
-				case DSI::FQueryDescription::EOperatorType::SimpleAny:
-					Presence = EMassFragmentPresence::Any;
-					break;
-				case DSI::FQueryDescription::EOperatorType::SimpleNone:
-					Presence = EMassFragmentPresence::None;
-					break;
-				default:
-					continue;
-				}
-
-				if (Operand->Type->IsChildOf(FMassTag::StaticStruct()))
-				{
-					NativeQuery.AddTagRequirement(*(Operand->Type), Presence);
-				}
-				else if (Operand->Type->IsChildOf(FMassFragment::StaticStruct()))
-				{
-					NativeQuery.AddRequirement(Operand->Type.Get(), EMassFragmentAccess::None, Presence);
-				}
+			case DSI::FQueryDescription::EOperatorType::SimpleAll:
+				Presence = EMassFragmentPresence::All;
+				break;
+			case DSI::FQueryDescription::EOperatorType::SimpleAny:
+				Presence = EMassFragmentPresence::Any;
+				break;
+			case DSI::FQueryDescription::EOperatorType::SimpleNone:
+				Presence = EMassFragmentPresence::None;
+				break;
+			default:
+				continue;
+			}
 
-				++Operand;
+			if (Operand->Type->IsChildOf(FMassTag::StaticStruct()))
+			{
+				NativeQuery.AddTagRequirement(*(Operand->Type), Presence);
 			}
-			return true;
+			else if (Operand->Type->IsChildOf(FMassFragment::StaticStruct()))
+			{
+				NativeQuery.AddRequirement(Operand->Type.Get(), EMassFragmentAccess::None, Presence);
+			}
+
+			++Operand;
 		}
+		return true;
 	}
 	return false;
 }
diff --git a/Engine/Source/Runtime/TypedElementFramework/Private/Elements/Framework/TypedElementQueryBuilder.cpp b/Engine/Source/Runtime/TypedElementFramework/Private/Elements/Framework/TypedElementQueryBuilder.cpp
--- a/Engine/Source/Runtime/TypedElementFramework/Private/Elements/Framework/TypedElementQueryBuilder.cpp
+++ b/Engine/Source/Runtime/TypedElementFramework/Private/Elements/Framework/TypedElementQueryBuilder.cpp
@@ -113,7 +113,6 @@ namespace TypedElementQueryBuilder
 	FSimpleQuery::FSimpleQuery(ITypedElementDataStorageInterface::FQueryDescription* Query)
 		: Query(Query)
 	{
-		Query->bSimpleQuery = true;
 	}
 
 	FSimpleQuery& FSimpleQuery::All(const UScriptStruct* Target)
diff --git a/Engine/Source/Runtime/TypedElementFramework/Public/Elements/Common/TypedElementQueryDescription.h b/Engine/Source/Runtime/TypedElementFramework/Public/Elements/Common/TypedElementQueryDescription.h
--- a/Engine/Source/Runtime/TypedElementFramework/Public/Elements/Common/TypedElementQueryDescription.h
+++ b/Engine/Source/Runtime/TypedElementFramework/Public/Elements/Common/TypedElementQueryDescription.h
@@ -94,7 +94,5 @@ namespace TypedElementDataStorage
 		FMetaData MetaData;
 
 		EActionType Action;
-		/** If true, this query only has simple operations and is guaranteed to be executed fully and at optimal performance. */
-		bool bSimpleQuery{ false };
 	};
 } // namespace TypedElementDataStorage
