SourceControl / ViewportStatus: Fixed regression introduced in #32614335 that resulted in the RevisionControlProcessor never applying the overlay color on the actors.

diff --git a/Engine/Plugins/Experimental/TypedElementsDataStorage/Source/TypedElementsDataStorageRevisionControl/Source/Private/RevisionControlProcessors.cpp b/Engine/Plugins/Experimental/TypedElementsDataStorage/Source/TypedElementsDataStorageRevisionControl/Source/Private/RevisionControlProcessors.cpp
--- a/Engine/Plugins/Experimental/TypedElementsDataStorage/Source/TypedElementsDataStorageRevisionControl/Source/Private/RevisionControlProcessors.cpp
+++ b/Engine/Plugins/Experimental/TypedElementsDataStorage/Source/TypedElementsDataStorageRevisionControl/Source/Private/RevisionControlProcessors.cpp
@@ -244,8 +244,6 @@ void UTypedElementRevisionControlFactory::RegisterApplyOverlays(ITypedElementDat
 				);
 			}
 		)
-		.Where()
-			.Any<FSCCNotCurrentTag, FSCCLockedTag>()
 		.DependsOn()
 			.SubQuery(ApplyOverlaysObjectToSCC)
 		.Compile()
