Memento support for undo/redo of deletion through TEDS.

With this change the memento system is involved with undo/redo when objects are created or deleted. This allows columns to be stored in a memento table when an object is deleted and the columns to be restored when the deletion is undone.Unlike with reinstancing, these memento rows remain until the transaction itself is destroyed.
The following additional changes were made:
- Simplified the memento system by replacing query callbacks with direct calls. This is now possible due to updates made to TEDS Core on the delayed column processing.
- The memento system is no longer a UObject but an object on the TEDS environment instead. This will allow more systems inside the TEDS plugin to use it in the future.
- The environment can now more easily be shared within the TEDS plugin.
- There's now a templated version of RegisterTable to make it easier to create tables.
- TEDS Core now has a `OnUpdateCompleted` callback that can be used to do any post-update updates such as clearing up temporary storage.

diff --git a/Engine/Plugins/Experimental/TypedElementsDataStorage/Source/TypedElementsDataStorage/Private/Compatibility/TypedElementObjectReinstancingManager.cpp b/Engine/Plugins/Experimental/TypedElementsDataStorage/Source/TypedElementsDataStorage/Private/Compatibility/TypedElementObjectReinstancingManager.cpp
--- a/Engine/Plugins/Experimental/TypedElementsDataStorage/Source/TypedElementsDataStorage/Private/Compatibility/TypedElementObjectReinstancingManager.cpp
+++ b/Engine/Plugins/Experimental/TypedElementsDataStorage/Source/TypedElementsDataStorage/Private/Compatibility/TypedElementObjectReinstancingManager.cpp
@@ -7,7 +7,6 @@
 #include "Memento/TypedElementMementoRowTypes.h"
 #include "TypedElementDatabase.h"
 #include "TypedElementDatabaseCompatibility.h"
-#include "Memento/TypedElementMementoInterface.h"
 
 DECLARE_LOG_CATEGORY_CLASS(LogTedsObjectReinstancing, Log, Log)
 
@@ -16,82 +15,49 @@ UTypedElementObjectReinstancingManager::UTypedElementObjectReinstancingManager()
 {
 }
 
-void UTypedElementObjectReinstancingManager::Initialize(UTypedElementDatabase& InDatabase, UTypedElementDatabaseCompatibility& InDataStorageCompatibility, UTypedElementMementoSystem& InMementoSystem)
+void UTypedElementObjectReinstancingManager::Initialize(UTypedElementDatabase& InDatabase, UTypedElementDatabaseCompatibility& InDataStorageCompatibility)
 {
 	Database = &InDatabase;
 	DataStorageCompatibility = &InDataStorageCompatibility;
-	MementoSystem = &InMementoSystem;
-
-	ReinstancingCallbackHandle = FCoreUObjectDelegates::OnObjectsReinstanced.AddUObject(this, &UTypedElementObjectReinstancingManager::HandleOnObjectsReinstanced);
-	ObjectRemovedCallbackHandle = DataStorageCompatibility->RegisterObjectRemovedCallback([this](const void* Object, const FTypedElementDatabaseCompatibilityObjectTypeInfo& TypeInfo, TypedElementRowHandle Row)
-	{
-		HandleOnObjectPreRemoved(Object, TypeInfo, Row);
-	});
 	
-	RegisterQueries();
+	UpdateCompletedCallbackHandle = 
+		Database->OnUpdateCompleted().AddUObject(this, &UTypedElementObjectReinstancingManager::UpdateCompleted);
+	ReinstancingCallbackHandle = 
+		FCoreUObjectDelegates::OnObjectsReinstanced.AddUObject(this, &UTypedElementObjectReinstancingManager::HandleOnObjectsReinstanced);
+	ObjectRemovedCallbackHandle = DataStorageCompatibility->RegisterObjectRemovedCallback(
+		[this](const void* Object, const FTypedElementDatabaseCompatibilityObjectTypeInfo& TypeInfo, TypedElementRowHandle Row)
+		{
+			HandleOnObjectPreRemoved(Object, TypeInfo, Row);
+		});
 }
 
 void UTypedElementObjectReinstancingManager::Deinitialize()
 {
-	UnregisterQueries();
-	
 	FCoreUObjectDelegates::OnObjectsReinstanced.Remove(ReinstancingCallbackHandle);
 	DataStorageCompatibility->UnregisterObjectRemovedCallback(ObjectRemovedCallbackHandle);
+	Database->OnUpdateCompleted().Remove(UpdateCompletedCallbackHandle);
 
-	MementoSystem = nullptr;
 	DataStorageCompatibility = nullptr;
 	Database = nullptr;
 }
 
-void UTypedElementObjectReinstancingManager::RegisterQueries()
+void UTypedElementObjectReinstancingManager::UpdateCompleted()
 {
-	using namespace TypedElementQueryBuilder;
-	using DSI = ITypedElementDataStorageInterface;
+	using namespace TypedElementDataStorage;
 
+	UTypedElementMementoSystem& MementoSystem = Database->GetEnvironment()->GetMementoSystem();
+	for (TMap<const void*, RowHandle>::TConstIterator It = OldObjectToMementoMap.CreateConstIterator(); It; ++It)
 	{
-		TypedElementQueryHandle QueryHandle = Database->RegisterQuery(
-		Select(
-		TEXT("Memento cleanup"),
-		FProcessor(DSI::EQueryTickPhase::FrameEnd, Database->GetQueryTickGroupName(DSI::EQueryTickGroups::Default)),
-			[this](TypedElementDataStorage::IQueryContext& Context, const TypedElementRowHandle* Row, const FTypedElementsReinstanceableSourceObject*)
-			{
-				Context.RemoveRows(Context.GetRowHandles());
-				
-				TConstArrayView<const FTypedElementsReinstanceableSourceObject> ReinstanceableArray = MakeArrayView(Context.GetColumn<FTypedElementsReinstanceableSourceObject>(), Context.GetRowCount());
-				for (const FTypedElementsReinstanceableSourceObject& Reinstanceable : ReinstanceableArray)
-				{
-					OldObjectToMementoMap.Remove(Reinstanceable.Object);
-				}
-				
-			})
-			.Where().All<FTypedElementMementoTag, FTypedElementMementoPopulated>()
-			.Compile()
-		);
-		check(QueryHandle != TypedElementInvalidQueryHandle);
+		MementoSystem.DestroyMemento(It.Value());
 	}
+	OldObjectToMementoMap.Reset();
 }
 
-void UTypedElementObjectReinstancingManager::UnregisterQueries()
-{
-	// TODO: Observer queries cannot be unregistered in Mass at this time
-}
-
-void UTypedElementObjectReinstancingManager::HandleOnObjectPreRemoved(const void* Object, const FTypedElementDatabaseCompatibilityObjectTypeInfo& TypeInfo, TypedElementRowHandle ObjectRow)
+void UTypedElementObjectReinstancingManager::HandleOnObjectPreRemoved(
+	const void* Object, const FTypedElementDatabaseCompatibilityObjectTypeInfo& TypeInfo, TypedElementRowHandle ObjectRow)
 {
 	// This is the chance to record the old object to memento
-	TypedElementRowHandle Memento = MementoSystem->CreateMemento(Database.Get());
-	
-	if (!ensureMsgf(Database->HasColumns(Memento, TConstArrayView<const UScriptStruct*>({FTypedElementMementoTag::StaticStruct()})),
-		TEXT("Memento System should create rows with the MementoTag")))
-	{
-		Database->RemoveRow(Memento);
-		return;
-	}
-
-	ITypedElementDataStorageInterface* Interface = Database.Get();
-	Interface->AddColumn(ObjectRow, FTypedElementMementoOnDelete{ .Memento = Memento });
-	Interface->AddColumn(Memento, FTypedElementsReinstanceableSourceObject{.Object = Object});
-
+	TypedElementRowHandle Memento = Database->GetEnvironment()->GetMementoSystem().CreateMemento(ObjectRow);
 	OldObjectToMementoMap.Add(Object, Memento);
 }
 
@@ -109,8 +75,6 @@ void UTypedElementObjectReinstancingManager::HandleOnObjectsReinstanced(
 			UObject* NewInstanceObject = Iter->Value;
 			if (NewInstanceObject == nullptr)
 			{
-				// Reinstancing resulted in no target object.  Delete the memento.
-				Interface->AddColumn<FTypedElementMementoReinstanceAborted>(Memento);
 				continue;
 			}
 			
@@ -123,7 +87,7 @@ void UTypedElementObjectReinstancingManager::HandleOnObjectsReinstanced(
 			}
 
 			// Kick off re-instantiation of NewObjectRow from the Memento
-			Interface->AddColumn(Memento, FTypedElementMementoReinstanceTarget{ .Target = NewObjectRow });
+			Database->GetEnvironment()->GetMementoSystem().RestoreMemento(Memento, NewObjectRow);
 		}
 	}
 }
diff --git a/Engine/Plugins/Experimental/TypedElementsDataStorage/Source/TypedElementsDataStorage/Private/Compatibility/TypedElementObjectReinstancingManager.h b/Engine/Plugins/Experimental/TypedElementsDataStorage/Source/TypedElementsDataStorage/Private/Compatibility/TypedElementObjectReinstancingManager.h
--- a/Engine/Plugins/Experimental/TypedElementsDataStorage/Source/TypedElementsDataStorage/Private/Compatibility/TypedElementObjectReinstancingManager.h
+++ b/Engine/Plugins/Experimental/TypedElementsDataStorage/Source/TypedElementsDataStorage/Private/Compatibility/TypedElementObjectReinstancingManager.h
@@ -1,4 +1,4 @@
-﻿// Copyright Epic Games, Inc. All Rights Reserved.
+// Copyright Epic Games, Inc. All Rights Reserved.
 
 #pragma once
 
@@ -21,41 +21,26 @@ class UTypedElementObjectReinstancingManager : public UObject
 public:
 	UTypedElementObjectReinstancingManager();
 
-	void Initialize(UTypedElementDatabase& InDatabase, UTypedElementDatabaseCompatibility& InDataStorageCompatibility, UTypedElementMementoSystem& InMementoSystem);
+	void Initialize(UTypedElementDatabase& InDatabase, UTypedElementDatabaseCompatibility& InDataStorageCompatibility);
 	void Deinitialize();
 
 private:
-	void RegisterQueries();
-	void UnregisterQueries();
-	void HandleOnObjectPreRemoved(const void* Object, const FTypedElementDatabaseCompatibilityObjectTypeInfo& TypeInfo, TypedElementRowHandle ObjectRow);
+	void UpdateCompleted();
+	void HandleOnObjectPreRemoved(
+		const void* Object, const FTypedElementDatabaseCompatibilityObjectTypeInfo& TypeInfo, TypedElementDataStorage::RowHandle ObjectRow);
 	void HandleOnObjectsReinstanced(const FCoreUObjectDelegates::FReplacementObjectMap& ObjectReplacementMap);
 
 	UPROPERTY()
 	TObjectPtr<UTypedElementDatabase> Database = nullptr;
 	UPROPERTY()
 	TObjectPtr<UTypedElementDatabaseCompatibility> DataStorageCompatibility = nullptr;
-	UPROPERTY()
-	TObjectPtr<UTypedElementMementoSystem> MementoSystem = nullptr;
-
+	
 	// Reverse lookup that holds all populated mementos for recently deleted objects
 	// Entry removed when the memento is removed
-	TMap<const void*, TypedElementRowHandle> OldObjectToMementoMap;
+	TMap<const void*, TypedElementDataStorage::RowHandle> OldObjectToMementoMap;
 	
-	TypedElementTableHandle MementoRowBaseTable;
+	TypedElementDataStorage::TableHandle MementoRowBaseTable;
+	FDelegateHandle UpdateCompletedCallbackHandle;
 	FDelegateHandle ReinstancingCallbackHandle;
 	FDelegateHandle ObjectRemovedCallbackHandle;
 };
-
-USTRUCT()
-struct FTypedElementsReinstanceableSourceObject : public FTypedElementDataStorageColumn
-{
-	GENERATED_BODY()
-
-	const void* Object;
-};
-
-USTRUCT()
-struct FTypedElementsObjectReinstanceSourceColumnInitialized : public FTypedElementDataStorageTag
-{
-	GENERATED_BODY()
-};
\ No newline at end of file
diff --git a/Engine/Plugins/Experimental/TypedElementsDataStorage/Source/TypedElementsDataStorage/Private/Memento/TypedElementMementoRowTypes.h b/Engine/Plugins/Experimental/TypedElementsDataStorage/Source/TypedElementsDataStorage/Private/Memento/TypedElementMementoRowTypes.h
--- a/Engine/Plugins/Experimental/TypedElementsDataStorage/Source/TypedElementsDataStorage/Private/Memento/TypedElementMementoRowTypes.h
+++ b/Engine/Plugins/Experimental/TypedElementsDataStorage/Source/TypedElementsDataStorage/Private/Memento/TypedElementMementoRowTypes.h
@@ -1,4 +1,4 @@
-﻿// Copyright Epic Games, Inc. All Rights Reserved.
+// Copyright Epic Games, Inc. All Rights Reserved.
 
 #pragma once
 
@@ -18,26 +18,3 @@ struct FTypedElementMementoTag : public FTypedElementDataStorageTag
 {
 	GENERATED_BODY()
 };
-
-/**
- * A column added to a memento row which will trigger the process
- * of converting memento columns to the matching mementoizable columns in
- * the target row
- */
-USTRUCT()
-struct FTypedElementMementoReinstanceTarget : public FTypedElementDataStorageColumn
-{
-	GENERATED_BODY()
-	TypedElementRowHandle Target;
-};
-
-/**
- * Tag added to memento row which will cause the memento to be deleted.  It is used
- * in the case where mementoization was aborted due to a missing target object
- * from the reinstancing callback
- */
-USTRUCT()
-struct FTypedElementMementoReinstanceAborted : public FTypedElementDataStorageTag
-{
-	GENERATED_BODY()
-};
\ No newline at end of file
diff --git a/Engine/Plugins/Experimental/TypedElementsDataStorage/Source/TypedElementsDataStorage/Private/Memento/TypedElementMementoSystem.cpp b/Engine/Plugins/Experimental/TypedElementsDataStorage/Source/TypedElementsDataStorage/Private/Memento/TypedElementMementoSystem.cpp
--- a/Engine/Plugins/Experimental/TypedElementsDataStorage/Source/TypedElementsDataStorage/Private/Memento/TypedElementMementoSystem.cpp
+++ b/Engine/Plugins/Experimental/TypedElementsDataStorage/Source/TypedElementsDataStorage/Private/Memento/TypedElementMementoSystem.cpp
@@ -4,192 +4,93 @@
 
 #include "Elements/Common/TypedElementDataStorageLog.h"
 #include "Elements/Framework/TypedElementQueryBuilder.h"
-#include "Memento/TypedElementMementoInterface.h"
 #include "Memento/TypedElementMementoTranslators.h"
 #include "StructArrayView.h"
 #include "TypedElementMementoRowTypes.h"
 #include "TypedElementDatabase.h"
 
-void UTypedElementMementoSystem::Initialize(UTypedElementDatabase& DataStorage)
-{
-	RegisterTables(DataStorage);
-	RegisterQueries(DataStorage);
-}
-
-void UTypedElementMementoSystem::Deinitialize()
-{
-}
-
-TypedElementRowHandle UTypedElementMementoSystem::CreateMemento(ITypedElementDataStorageInterface* DataStorage)
-{
-	return DataStorage->AddRow(MementoRowBaseTable);
-}
-
-void UTypedElementMementoSystem::RegisterTables(UTypedElementDatabase& DataStorage)
+UTypedElementMementoSystem::UTypedElementMementoSystem(ITypedElementDataStorageInterface& InDataStorage)
+	: DataStorage(InDataStorage)
 {
 	// Register tables that will be used by reinstancing
-	MementoRowBaseTable = DataStorage.RegisterTable(
-		{FTypedElementMementoTag::StaticStruct()},
-		TEXT("MementoRowBaseTable"));
-	
-}
-
-void UTypedElementMementoSystem::RegisterQueries(UTypedElementDatabase& DataStorage) const
-{
-	using DSI = ITypedElementDataStorageInterface;
-	
-	TArray<const UTypedElementMementoTranslatorBase*> MementoTranslators;
+	MementoRowBaseTable = DataStorage.RegisterTable<FTypedElementMementoTag>(TEXT("MementoRowBaseTable"));
 
 	// Discover all MementoTranslators
 	{
-		const bool bIncludeDerived = true;
-		EObjectFlags ExcludeFlags = EObjectFlags::RF_NoFlags;
-		ForEachObjectOfClass(UTypedElementMementoTranslatorBase::StaticClass(), [&MementoTranslators](UObject* Object)
-		{
-			const UTypedElementMementoTranslatorBase* TranslatorCandidate = Cast<UTypedElementMementoTranslatorBase>(Object);
-			// Exclude abstract classes
-			if (TranslatorCandidate->GetClass()->GetClassFlags() & EClassFlags::CLASS_Abstract)
+		constexpr bool bIncludeDerived = true;
+		constexpr EObjectFlags ExcludeFlags = EObjectFlags::RF_NoFlags;
+		ForEachObjectOfClass(UTypedElementMementoTranslatorBase::StaticClass(),
+			[this](UObject* Object)
 			{
-				return;
-			}
-			MementoTranslators.Add(TranslatorCandidate);
-		},
-		bIncludeDerived,
-		ExcludeFlags);
+				const UTypedElementMementoTranslatorBase* TranslatorCandidate = Cast<UTypedElementMementoTranslatorBase>(Object);
+				// Exclude abstract classes
+				if (TranslatorCandidate->GetClass()->GetClassFlags() & EClassFlags::CLASS_Abstract)
+				{
+					return;
+				}
+				MementoTranslators.Add(TranslatorCandidate);
+			},
+			bIncludeDerived, ExcludeFlags);
 	}
+}
 
-	using namespace TypedElementQueryBuilder;
+TypedElementRowHandle UTypedElementMementoSystem::CreateMemento(TypedElementDataStorage::RowHandle SourceRow)
+{
+	using namespace TypedElementDataStorage;
 
-	// Setup observer queries to execute memento translators on row deletions that match the translator's
-	// capabilities.
-	for (int32 Index = 0, End = MementoTranslators.Num(); Index < End; ++Index)
+	RowHandle MementoRow = DataStorage.AddRow(MementoRowBaseTable);
+	
+	for (const UTypedElementMementoTranslatorBase* Translator : MementoTranslators)
 	{
-		const UTypedElementMementoTranslatorBase* MementoTranslator = MementoTranslators[Index];
-		const UScriptStruct* MementoizedColumn = MementoTranslator->GetColumnType();
-		const UScriptStruct* Memento = MementoTranslator->GetMementoType();
-		
-		const FName TranslationProcessorName = FName(FString::Printf(TEXT("MementoTranslator (Column->Memento) %s -> %s"), *MementoizedColumn->GetName(), *Memento->GetName()));
-		const TypedElementQueryHandle QueryHandle = DataStorage.RegisterQuery(
-			Select(
-				TranslationProcessorName,
-				FObserver::OnRemove<FTypedElementMementoOnDelete>(),
-				[MementoTranslator](TypedElementDataStorage::IQueryContext& Context, const TypedElementRowHandle* Rows, const FTypedElementMementoOnDelete* MementoTargets)
+		if (void* SourceColumn = DataStorage.GetColumnData(SourceRow, Translator->GetColumnType()))
+		{
+			const UScriptStruct* MementoType = Translator->GetMementoType();
+			DataStorage.AddColumnData(MementoRow, MementoType,
+				[Translator, SourceColumn](void* MementoColumn, const UScriptStruct& ColumnType)
 				{
-					const int32 RowCount = static_cast<int32>(Context.GetRowCount());					
-					const UScriptStruct* MementoType = MementoTranslator->GetMementoType();
-					const UScriptStruct* MementoizedColumn = MementoTranslator->GetColumnType();
-
-					TArrayView<const TypedElementRowHandle> RowsView = MakeArrayView(Rows, RowCount);
-					TArrayView<const FTypedElementMementoOnDelete> MementoTargetsView = MakeArrayView(MementoTargets, RowCount);
-					FConstStructArrayView SourceColumnArrayView(*MementoizedColumn, Context.GetColumn(MementoizedColumn), RowCount);
-
-					for (int32 Index = 0; Index < RowCount; ++Index)
-					{
-						TypedElementRowHandle SourceRow = RowsView[Index];
-						TypedElementRowHandle TargetMementoRow = MementoTargetsView[Index].Memento;
-						
-						void* TargetMementoStagedColumn = Context.AddColumnUninitialized(TargetMementoRow, MementoType);
-						// StagedMementoColumn will be uninitialized memory, ensure constructor called
-						MementoType->InitializeStruct(TargetMementoStagedColumn);
-		
-						const void* SourceColumn = SourceColumnArrayView[Index].GetMemory();
-
-						UE_LOG(LogTypedElementDataStorage, VeryVerbose, TEXT("Column->Memento: %llu -> %llu   0x%p -> 0x%p"), SourceRow, TargetMementoRow, SourceColumn, TargetMementoStagedColumn);
+					ColumnType.InitializeStruct(MementoColumn);
+					Translator->TranslateColumnToMemento(SourceColumn, MementoColumn);
+				},
+				[](const UScriptStruct& ColumnType, void* Destination, void* Source)
+				{
+					ColumnType.CopyScriptStruct(Destination, Source);
+				});
 
-						MementoTranslator->TranslateColumnToMemento(SourceColumn, TargetMementoStagedColumn);
-					}
-				})
-				.ReadOnly(MementoizedColumn)
-				.Compile());
-		check(QueryHandle != TypedElementInvalidQueryHandle);
+			UE_LOG(LogTypedElementDataStorage, VeryVerbose, 
+				TEXT("Column->Memento: %llu -> %llu"), SourceRow, MementoRow);
+		}
 	}
 
-	{
-		const TypedElementQueryHandle QueryHandle = DataStorage.RegisterQuery(
-			Select(
-				TEXT("Add Populated To Memento"),
-				FObserver::OnRemove<FTypedElementMementoOnDelete>(),
-				[](TypedElementDataStorage::IQueryContext& Context, TypedElementRowHandle Memento, const FTypedElementMementoOnDelete& MementoRow)
-				{
-					Context.AddColumns(MementoRow.Memento, TConstArrayView<const UScriptStruct*>({FTypedElementMementoPopulated::StaticStruct()}));
-				})
-		.Compile());
-		check(QueryHandle != TypedElementInvalidQueryHandle);
-	}
+	return MementoRow;
+}
 
-	// Setup processors to execute memento translators on reinstantiate columns onto a target
-	// The primary query runs over all mementos with a reinstance target and a column holding memento data
-	// The subquery runs on rows targeted by the reinstance column that also have a mementoizable column
-	// Note: Currently it is not possible for mementos to create a new column from a subquery
-	for (int32 Index = 0, End = MementoTranslators.Num(); Index < End; ++Index)
+void UTypedElementMementoSystem::RestoreMemento(TypedElementDataStorage::RowHandle MementoRow, TypedElementDataStorage::RowHandle TargetRow)
+{
+	for (const UTypedElementMementoTranslatorBase* Translator : MementoTranslators)
 	{
-		const UTypedElementMementoTranslatorBase* MementoTranslator = MementoTranslators[Index];
-		const UScriptStruct* MementoizedColumnType = MementoTranslator->GetColumnType();
-		const UScriptStruct* MementoType = MementoTranslator->GetMementoType();
-
-		namespace TEDS = TypedElementDataStorage;
-
-		const FName TranslationProcessorName = FName(FString::Printf(TEXT("MementoTranslator (Memento->Column) %s -> %s"), *MementoType->GetName(), *MementoizedColumnType->GetName()));
-		const TypedElementQueryHandle QueryHandleForExistence = DataStorage.RegisterQuery(
-			Select(
-				TranslationProcessorName,
-				FProcessor(DSI::EQueryTickPhase::PostPhysics, DataStorage.GetQueryTickGroupName(DSI::EQueryTickGroups::Default)),
-				[MementoTranslator](TEDS::IQueryContext& Context, const TypedElementRowHandle* Rows, const FTypedElementMementoReinstanceTarget* ReinstanceTargets)
+		if (void* MementoColumn = DataStorage.GetColumnData(MementoRow, Translator->GetMementoType()))
+		{
+			const UScriptStruct* TargetType = Translator->GetColumnType();
+			DataStorage.AddColumnData(TargetRow, TargetType,
+				[Translator, MementoColumn](void* TargetColumn, const UScriptStruct& ColumnType)
 				{
-					const int32 RowCount = static_cast<int32>(Context.GetRowCount());					
-					const UScriptStruct* MementoType = MementoTranslator->GetMementoType();
-					const UScriptStruct* ColumnType = MementoTranslator->GetColumnType();
-
-					TArrayView<const TypedElementRowHandle> RowsView = MakeArrayView(Rows, RowCount);
-					TArrayView<const FTypedElementMementoReinstanceTarget> ReinstanceTargetsView = MakeArrayView(ReinstanceTargets, RowCount);
-					FConstStructArrayView MementoArrayView(*MementoType, Context.GetColumn(MementoType), RowCount);
-
-					for (int32 Index = 0; Index < RowCount; ++Index)
-					{
-						const TypedElementRowHandle Row = RowsView[Index];
-						const TypedElementRowHandle Target = ReinstanceTargetsView[Index].Target;
-						
-						// Add the column, note the column may already exist.
-						void* TargetStagedColumn = Context.AddColumnUninitialized(Target, ColumnType);
-
-						// StagedColumn will be uninitialized memory, ensure constructor called
-						MementoTranslator->GetColumnType()->InitializeStruct(TargetStagedColumn);
-						
-						const void* SourceMementoOfColumn = MementoArrayView[Index].GetMemory();
-
-						UE_LOG(LogTypedElementDataStorage, VeryVerbose, TEXT("Memento->Column: %llu -> %llu   0x%p -> 0x%p"), Row, Target, SourceMementoOfColumn, TargetStagedColumn);
-						
-						MementoTranslator->TranslateMementoToColumn(SourceMementoOfColumn, TargetStagedColumn);
-					}
-				})
-				.ReadOnly(MementoType)
-				.Where()
-					.All<FTypedElementMementoTag>()
-				.Compile());
-		check(QueryHandleForExistence != TypedElementInvalidQueryHandle);
-	}
-	
-	/**
-	 * A processor which deletes mementos that were used for reinstancing another row
-	 */
-	{
-		const TypedElementQueryHandle QueryHandle = DataStorage.RegisterQuery(
-			Select(
-			TEXT("Delete mementos used for reinstancing"),
-			FProcessor(DSI::EQueryTickPhase::FrameEnd, DataStorage.GetQueryTickGroupName(DSI::EQueryTickGroups::Default)),
-			[](TypedElementDataStorage::IQueryContext& Context, const TypedElementRowHandle* Rows)
+					ColumnType.InitializeStruct(TargetColumn);
+					Translator->TranslateMementoToColumn(MementoColumn, TargetColumn);
+				},
+				[](const UScriptStruct& ColumnType, void* Destination, void* Source)
 				{
-					TConstArrayView<TypedElementRowHandle> RowView(Rows, Context.GetRowCount());
-					Context.RemoveRows(RowView);
-				})
-				.Where()
-					.All<FTypedElementMementoTag>()
-					.Any<FTypedElementMementoReinstanceTarget, FTypedElementMementoReinstanceAborted>()
-				.Compile()
-			);
-		check(QueryHandle != TypedElementInvalidQueryHandle);
+					ColumnType.CopyScriptStruct(Destination, Source);
+				});
+
+			UE_LOG(LogTypedElementDataStorage, VeryVerbose, 
+				TEXT("Memento->Column: %llu -> %llu"), MementoRow, TargetRow);
+		}
 	}
 }
 
-
-
+void UTypedElementMementoSystem::DestroyMemento(TypedElementDataStorage::RowHandle MementoRow)
+{
+	checkf(DataStorage.IsRowAvailable(MementoRow) && DataStorage.HasColumns<FTypedElementMementoTag>(MementoRow),
+		TEXT("Deleting memento row that's not marked as such."));
+	DataStorage.RemoveRow(MementoRow);
+}
diff --git a/Engine/Plugins/Experimental/TypedElementsDataStorage/Source/TypedElementsDataStorage/Private/Memento/TypedElementMementoSystem.h b/Engine/Plugins/Experimental/TypedElementsDataStorage/Source/TypedElementsDataStorage/Private/Memento/TypedElementMementoSystem.h
--- a/Engine/Plugins/Experimental/TypedElementsDataStorage/Source/TypedElementsDataStorage/Private/Memento/TypedElementMementoSystem.h
+++ b/Engine/Plugins/Experimental/TypedElementsDataStorage/Source/TypedElementsDataStorage/Private/Memento/TypedElementMementoSystem.h
@@ -1,32 +1,25 @@
-﻿// Copyright Epic Games, Inc. All Rights Reserved.
+// Copyright Epic Games, Inc. All Rights Reserved.
 
 #pragma once
 
-#include "Elements/Interfaces/TypedElementDataStorageFactory.h"
-#include "Elements/Columns/TypedElementSelectionColumns.h"
-#include "UObject/UObjectGlobals.h"
-
-#include "TypedElementMementoSystem.generated.h"
+#include "Containers/Array.h"
+#include "Elements/Common/TypedElementHandles.h"
 
 class ITypedElementDataStorageInterface;
-class IConsoleVariable;
-struct FTypedElementDatabaseCompatibilityObjectTypeInfo;
 class UTypedElementDatabase;
-
-UCLASS()
-class UTypedElementMementoSystem : public UObject
+class UTypedElementMementoTranslatorBase;
+	
+class UTypedElementMementoSystem
 {
-	GENERATED_BODY()
 public:
-
-	void Initialize(UTypedElementDatabase& DataStorage);
-	void Deinitialize();
-
-	TypedElementRowHandle CreateMemento(ITypedElementDataStorageInterface* DataStorage);
+	explicit UTypedElementMementoSystem(ITypedElementDataStorageInterface& InDataStorage);
+	
+	TypedElementDataStorage::RowHandle CreateMemento(TypedElementDataStorage::RowHandle SourceRow);
+	void RestoreMemento(TypedElementDataStorage::RowHandle MementoRow,  TypedElementDataStorage::RowHandle TargetRow);
+	void DestroyMemento(TypedElementDataStorage::RowHandle MementoRow);
 
 private:
-	void RegisterQueries(UTypedElementDatabase& DataStorage) const;
-	void RegisterTables(UTypedElementDatabase& DataStorage);
-	
-	TypedElementTableHandle MementoRowBaseTable;
+	TArray<const UTypedElementMementoTranslatorBase*> MementoTranslators;
+	TypedElementDataStorage::TableHandle MementoRowBaseTable;
+	ITypedElementDataStorageInterface& DataStorage;
 };
diff --git a/Engine/Plugins/Experimental/TypedElementsDataStorage/Source/TypedElementsDataStorage/Private/TypedElementDatabase.cpp b/Engine/Plugins/Experimental/TypedElementsDataStorage/Source/TypedElementsDataStorage/Private/TypedElementDatabase.cpp
--- a/Engine/Plugins/Experimental/TypedElementsDataStorage/Source/TypedElementsDataStorage/Private/TypedElementDatabase.cpp
+++ b/Engine/Plugins/Experimental/TypedElementsDataStorage/Source/TypedElementsDataStorage/Private/TypedElementDatabase.cpp
@@ -111,7 +111,7 @@ void UTypedElementDatabase::Initialize()
 	ActiveEditorPhaseManager = Mass->GetMutablePhaseManager();
 	if (ActiveEditorEntityManager && ActiveEditorPhaseManager)
 	{
-		Environment = MakeUnique<FTypedElementDatabaseEnvironment>(*ActiveEditorEntityManager, *ActiveEditorPhaseManager);
+		Environment = MakeShared<FTypedElementDatabaseEnvironment>(*this, *ActiveEditorEntityManager, *ActiveEditorPhaseManager);
 
 		using PhaseType = std::underlying_type_t<EQueryTickPhase>;
 		for (PhaseType PhaseId = 0; PhaseId < static_cast<PhaseType>(EQueryTickPhase::Max); ++PhaseId)
@@ -248,6 +248,7 @@ void UTypedElementDatabase::OnPostMassTick(float DeltaTime)
 	checkf(IsAvailable(), TEXT("Typed Element Database was ticked while it's not ready."));
 	
 	Environment->NextUpdateCycle();
+	OnUpdateCompletedDelegate.Broadcast();
 }
 
 TSharedPtr<FMassEntityManager> UTypedElementDatabase::GetActiveMutableEditorEntityManager()
@@ -918,6 +919,11 @@ FTypedElementOnDataStorageUpdate& UTypedElementDatabase::OnUpdate()
 	return OnUpdateDelegate;
 }
 
+FTypedElementOnDataStorageUpdate& UTypedElementDatabase::OnUpdateCompleted()
+{
+	return OnUpdateCompletedDelegate;
+}
+
 bool UTypedElementDatabase::IsAvailable() const
 {
 	return bool(ActiveEditorEntityManager);
@@ -969,6 +975,16 @@ void UTypedElementDatabase::Reset()
 	ActiveEditorEntityManager.Reset();
 }
 
+TSharedPtr<FTypedElementDatabaseEnvironment> UTypedElementDatabase::GetEnvironment()
+{
+	return Environment;
+}
+
+TSharedPtr<const FTypedElementDatabaseEnvironment> UTypedElementDatabase::GetEnvironment() const
+{
+	return Environment;
+}
+
 void UTypedElementDatabase::DebugPrintQueryCallbacks(FOutputDevice& Output)
 {
 	Environment->GetQueryStore().DebugPrintQueryCallbacks(Output);
diff --git a/Engine/Plugins/Experimental/TypedElementsDataStorage/Source/TypedElementsDataStorage/Private/TypedElementDatabase.h b/Engine/Plugins/Experimental/TypedElementsDataStorage/Source/TypedElementsDataStorage/Private/TypedElementDatabase.h
--- a/Engine/Plugins/Experimental/TypedElementsDataStorage/Source/TypedElementsDataStorage/Private/TypedElementDatabase.h
+++ b/Engine/Plugins/Experimental/TypedElementsDataStorage/Source/TypedElementsDataStorage/Private/TypedElementDatabase.h
@@ -135,8 +135,12 @@ public:
 	void RemoveIndex(TypedElementDataStorage::IndexHash Index) override;
 
 	FTypedElementOnDataStorageUpdate& OnUpdate() override;
+	FTypedElementOnDataStorageUpdate& OnUpdateCompleted() override;
 	bool IsAvailable() const override;
 	void* GetExternalSystemAddress(UClass* Target) override;
+	
+	TSharedPtr<FTypedElementDatabaseEnvironment> GetEnvironment();
+	TSharedPtr<const FTypedElementDatabaseEnvironment> GetEnvironment() const;
 
 	void DebugPrintQueryCallbacks(FOutputDevice& Output);
 
@@ -167,9 +171,10 @@ private:
 	// Ordered array of factories by the return value of GetOrder()
 	TArray<FFactoryTypePair> Factories;
 
-	TUniquePtr<FTypedElementDatabaseEnvironment> Environment;
+	TSharedPtr<FTypedElementDatabaseEnvironment> Environment;
 	
 	FTypedElementOnDataStorageUpdate OnUpdateDelegate;
+	FTypedElementOnDataStorageUpdate OnUpdateCompletedDelegate;
 	FDelegateHandle OnPreMassTickHandle;
 	FDelegateHandle OnPostMassTickHandle;
 
diff --git a/Engine/Plugins/Experimental/TypedElementsDataStorage/Source/TypedElementsDataStorage/Private/TypedElementDatabaseCompatibility.cpp b/Engine/Plugins/Experimental/TypedElementsDataStorage/Source/TypedElementsDataStorage/Private/TypedElementDatabaseCompatibility.cpp
--- a/Engine/Plugins/Experimental/TypedElementsDataStorage/Source/TypedElementsDataStorage/Private/TypedElementDatabaseCompatibility.cpp
+++ b/Engine/Plugins/Experimental/TypedElementsDataStorage/Source/TypedElementsDataStorage/Private/TypedElementDatabaseCompatibility.cpp
@@ -16,16 +16,19 @@
 #include "Elements/Framework/TypedElementIndexHasher.h"
 #include "Elements/Framework/TypedElementQueryBuilder.h"
 #include "MassActorSubsystem.h"
+#include "Memento/TypedElementMementoRowTypes.h"
+#include "TypedElementDatabase.h"
 #include "TypedElementDataStorageProfilingMacros.h"
 
-void UTypedElementDatabaseCompatibility::Initialize(ITypedElementDataStorageInterface* StorageInterface)
+void UTypedElementDatabaseCompatibility::Initialize(UTypedElementDatabase* InStorage)
 {
-	checkf(StorageInterface, TEXT("Typed Element's Database compatibility manager is being initialized with an invalid storage target."));
+	checkf(InStorage, TEXT("Typed Element's Database compatibility manager is being initialized with an invalid storage target."));
 	
-	Storage = StorageInterface;
+	Storage = InStorage;
+	Environment = InStorage->GetEnvironment();
 	Prepare();
 
-	StorageInterface->OnUpdate().AddUObject(this, &UTypedElementDatabaseCompatibility::Tick);
+	InStorage->OnUpdate().AddUObject(this, &UTypedElementDatabaseCompatibility::Tick);
 
 	PreEditChangePropertyDelegateHandle = FCoreUObjectDelegates::OnPreObjectPropertyChanged.AddUObject(this, &UTypedElementDatabaseCompatibility::OnPrePropertyChanged);
 	PostEditChangePropertyDelegateHandle = FCoreUObjectDelegates::OnObjectPropertyChanged.AddUObject(this, &UTypedElementDatabaseCompatibility::OnPostEditChangeProperty);
@@ -272,7 +275,7 @@ TypedElementRowHandle UTypedElementDatabaseCompatibility::AddCompatibleObjectExp
 		{
 			if (GUndo)
 			{
-				GUndo->StoreUndo(this, MakeUnique<FRegistrationCommandChange>(Object));
+				GUndo->StoreUndo(this, MakeUnique<FRegistrationCommandChange>(this, Object));
 			}
 		}
 	}
@@ -301,7 +304,7 @@ void UTypedElementDatabaseCompatibility::RemoveCompatibleObjectExplicitTransacti
 			{
 				if (GUndo)
 				{
-					GUndo->StoreUndo(this, MakeUnique<FDeregistrationCommandChange>(Object));
+					GUndo->StoreUndo(this, MakeUnique<FDeregistrationCommandChange>(this, Object));
 				}
 			}
 		}
@@ -769,29 +772,54 @@ SIZE_T GetTypeHash(const UTypedElementDatabaseCompatibility::FSyncTagInfo& Colum
 // UTypedElementDatabaseCompatibility::FRegistrationCommandChange
 //
 
-UTypedElementDatabaseCompatibility::FRegistrationCommandChange::FRegistrationCommandChange(UObject* InTargetObject)
-	: TargetObject(InTargetObject)
+UTypedElementDatabaseCompatibility::FRegistrationCommandChange::FRegistrationCommandChange(
+	UTypedElementDatabaseCompatibility* InOwner, UObject* InTargetObject)
+	: Owner(InOwner)
+	, TargetObject(InTargetObject)
 {
 }
 
+UTypedElementDatabaseCompatibility::FRegistrationCommandChange::~FRegistrationCommandChange()
+{
+	// If there has been no revert operation, there's also no memento.
+	if (UTypedElementDatabaseCompatibility* DataStorageCompat = Owner.Get(); 
+		DataStorageCompat && DataStorageCompat->Storage->IsRowAvailable(MementoRow))
+	{
+		DataStorageCompat->Environment->GetMementoSystem().DestroyMemento(MementoRow);
+	}
+}
+
 void UTypedElementDatabaseCompatibility::FRegistrationCommandChange::Apply(UObject* Object)
 {
-	if (UTypedElementDatabaseCompatibility* CompatibilityLayer = Cast<UTypedElementDatabaseCompatibility>(Object))
+	using namespace TypedElementDataStorage;
+
+	checkf(Owner.IsValid() && Owner.Get() == Object, 
+		TEXT("Applying registration transaction command within TEDS Compat was called after TEDS is not longer available."));
+	UTypedElementDatabaseCompatibility* DataStorageCompat = Owner.Get();
+	if (UObject* TargetRetrieved = TargetObject.Get(/*bEvenIfPendingKill=*/ true))
 	{
-		if (UObject* TargetRetrieved = TargetObject.Get(/*bEvenIfPendingKill=*/ true))
-		{
-			CompatibilityLayer->AddCompatibleObjectExplicitTransactionable<false>(TargetRetrieved);
-		}
+		RowHandle ObjectRow = DataStorageCompat->AddCompatibleObjectExplicitTransactionable<false>(TargetRetrieved);
+		DataStorageCompat->Environment->GetMementoSystem().RestoreMemento(MementoRow, ObjectRow);
 	}
 }
 
 void UTypedElementDatabaseCompatibility::FRegistrationCommandChange::Revert(UObject* Object)
 {
-	if (UTypedElementDatabaseCompatibility* CompatibilityLayer = Cast<UTypedElementDatabaseCompatibility>(Object))
+	using namespace TypedElementDataStorage;
+	
+	checkf(Owner.IsValid() && Owner.Get() == Object,
+		TEXT("Reverting registration transaction command within TEDS Compat was called after TEDS is not longer available."));
+
+	if (UObject* TargetRetrieved = TargetObject.Get(/*bEvenIfPendingKill=*/ true))
 	{
-		if (UObject* TargetRetrieved = TargetObject.Get(/*bEvenIfPendingKill=*/ true))
+		UTypedElementDatabaseCompatibility* DataStorageCompat = Owner.Get(); 
+		ITypedElementDataStorageInterface* DataStorage = DataStorageCompat->Storage;
+			
+		RowHandle ObjectRow = DataStorageCompat->FindRowWithCompatibleObjectExplicit(TargetRetrieved);
+		if (DataStorage->IsRowAvailable(ObjectRow))
 		{
-			CompatibilityLayer->RemoveCompatibleObjectExplicitTransactionable<false>(TargetRetrieved);
+			MementoRow = DataStorageCompat->Environment->GetMementoSystem().CreateMemento(ObjectRow);
+			DataStorageCompat->RemoveCompatibleObjectExplicitTransactionable<false>(TargetRetrieved);
 		}
 	}
 }
@@ -806,30 +834,57 @@ FString UTypedElementDatabaseCompatibility::FRegistrationCommandChange::ToString
 // UTypedElementDatabaseCompatibility::FDeregistrationCommandChange
 //
 
-UTypedElementDatabaseCompatibility::FDeregistrationCommandChange::FDeregistrationCommandChange(UObject* InTargetObject)
-	: TargetObject(InTargetObject)
+UTypedElementDatabaseCompatibility::FDeregistrationCommandChange::FDeregistrationCommandChange(
+	UTypedElementDatabaseCompatibility* InOwner, UObject* InTargetObject)
+	: Owner(InOwner)
+	, TargetObject(InTargetObject)
 {
+	using namespace TypedElementDataStorage;
+
+	ITypedElementDataStorageInterface* DataStorage = InOwner->Storage;
+
+	RowHandle ObjectRow = InOwner->FindRowWithCompatibleObjectExplicit(InTargetObject);
+	if (DataStorage->IsRowAvailable(ObjectRow))
+	{
+		MementoRow = InOwner->Environment->GetMementoSystem().CreateMemento(ObjectRow);
+	}
+}
+
+UTypedElementDatabaseCompatibility::FDeregistrationCommandChange::~FDeregistrationCommandChange()
+{
+	// There's no memento row if target object was never registered with TEDS Compat.
+	if (UTypedElementDatabaseCompatibility* DataStorageCompat = Owner.Get();
+		DataStorageCompat && DataStorageCompat->Storage->IsRowAvailable(MementoRow))
+	{
+		DataStorageCompat->Environment->GetMementoSystem().DestroyMemento(MementoRow);
+	}
 }
 
 void UTypedElementDatabaseCompatibility::FDeregistrationCommandChange::Apply(UObject* Object)
 {
-	if (UTypedElementDatabaseCompatibility* CompatibilityLayer = Cast<UTypedElementDatabaseCompatibility>(Object))
+	using namespace TypedElementDataStorage;
+
+	checkf(Owner.IsValid() && Owner.Get() == Object,
+		TEXT("Applying deregistration transaction command within TEDS Compat was called after TEDS is not longer available."));
+	UTypedElementDatabaseCompatibility* DataStorageCompat = Owner.Get();
+	if (UObject* TargetRetrieved = TargetObject.Get(/*bEvenIfPendingKill=*/ true))
 	{
-		if (UObject* TargetRetrieved = TargetObject.Get(/*bEvenIfPendingKill=*/ true))
-		{
-			CompatibilityLayer->RemoveCompatibleObjectExplicitTransactionable<false>(TargetRetrieved);
-		}
+		DataStorageCompat->RemoveCompatibleObjectExplicitTransactionable<false>(TargetRetrieved);
 	}
 }
 
 void UTypedElementDatabaseCompatibility::FDeregistrationCommandChange::Revert(UObject* Object)
 {
-	if (UTypedElementDatabaseCompatibility* CompatibilityLayer = Cast<UTypedElementDatabaseCompatibility>(Object))
+	using namespace TypedElementDataStorage;
+
+	checkf(Owner.IsValid() && Owner.Get() == Object,
+		TEXT("Reverting deregistration transaction command within TEDS Compat was called after TEDS is not longer available."));
+	
+	UTypedElementDatabaseCompatibility* DataStorageCompat = Owner.Get();
+	if (UObject* TargetRetrieved = TargetObject.Get(/*bEvenIfPendingKill=*/ true))
 	{
-		if (UObject* TargetRetrieved = TargetObject.Get(/*bEvenIfPendingKill=*/ true))
-		{
-			CompatibilityLayer->AddCompatibleObjectExplicitTransactionable<false>(TargetRetrieved);
-		}
+		RowHandle ObjectRow = DataStorageCompat->AddCompatibleObjectExplicitTransactionable<false>(TargetRetrieved);
+		DataStorageCompat->Environment->GetMementoSystem().RestoreMemento(MementoRow, ObjectRow);
 	}
 }
 
diff --git a/Engine/Plugins/Experimental/TypedElementsDataStorage/Source/TypedElementsDataStorage/Private/TypedElementDatabaseCompatibility.h b/Engine/Plugins/Experimental/TypedElementsDataStorage/Source/TypedElementsDataStorage/Private/TypedElementDatabaseCompatibility.h
--- a/Engine/Plugins/Experimental/TypedElementsDataStorage/Source/TypedElementsDataStorage/Private/TypedElementDatabaseCompatibility.h
+++ b/Engine/Plugins/Experimental/TypedElementsDataStorage/Source/TypedElementsDataStorage/Private/TypedElementDatabaseCompatibility.h
@@ -19,6 +19,8 @@
 
 class ITypedElementDataStorageInterface;
 struct FMassActorManager;
+class FTypedElementDatabaseEnvironment;
+class UTypedElementMementoSystem;
 
 enum class ETypedElementDatabaseCompatibilityObjectType : uint8;
 struct FTypedElementDatabaseCompatibilityObjectTypeInfo;
@@ -35,7 +37,7 @@ public:
 	
 	~UTypedElementDatabaseCompatibility() override = default;
 
-	void Initialize(ITypedElementDataStorageInterface* StorageInterface);
+	void Initialize(UTypedElementDatabase* InStorage);
 	void Deinitialize();
 
 	void RegisterRegistrationFilter(ObjectRegistrationFilter Filter) override;
@@ -63,26 +65,32 @@ private:
 	class FRegistrationCommandChange final : public FCommandChange
 	{
 	public:
-		explicit FRegistrationCommandChange(UObject* InTargetObject);
+		FRegistrationCommandChange(UTypedElementDatabaseCompatibility* InOwner, UObject* InTargetObject);
+		~FRegistrationCommandChange() override;
 
 		void Apply(UObject* Object) override;
 		void Revert(UObject* Object) override;
 		FString ToString() const override;
 
 	private:
+		TWeakObjectPtr<UTypedElementDatabaseCompatibility> Owner;
 		TWeakObjectPtr<UObject> TargetObject;
+		TypedElementDataStorage::RowHandle MementoRow = TypedElementDataStorage::InvalidRowHandle;
 	};
 	class FDeregistrationCommandChange final : public FCommandChange
 	{
 	public:
-		explicit FDeregistrationCommandChange(UObject* InTargetObject);
+		FDeregistrationCommandChange(UTypedElementDatabaseCompatibility* InOwner, UObject* InTargetObject);
+		~FDeregistrationCommandChange() override;
 
 		void Apply(UObject* Object) override;
 		void Revert(UObject* Object) override;
 		FString ToString() const override;
 
 	private:
+		TWeakObjectPtr<UTypedElementDatabaseCompatibility> Owner;
 		TWeakObjectPtr<UObject> TargetObject;
+		TypedElementDataStorage::RowHandle MementoRow = TypedElementDataStorage::InvalidRowHandle;
 	};
 
 	void Prepare();
@@ -212,7 +220,8 @@ private:
 	FDelegateHandle PostWorldInitializationDelegateHandle;
 	FDelegateHandle PreWorldFinishDestroyDelegateHandle;
 	FDelegateHandle ObjectReinstancedDelegateHandle;
-
+	
+	TSharedPtr<FTypedElementDatabaseEnvironment> Environment;
 	TypedElementDataStorage::QueryHandle ClassTypeInfoQuery;
 	TypedElementDataStorage::QueryHandle ScriptStructTypeInfoQuery;
 };
diff --git a/Engine/Plugins/Experimental/TypedElementsDataStorage/Source/TypedElementsDataStorage/Private/TypedElementDatabaseEnvironment.cpp b/Engine/Plugins/Experimental/TypedElementsDataStorage/Source/TypedElementsDataStorage/Private/TypedElementDatabaseEnvironment.cpp
--- a/Engine/Plugins/Experimental/TypedElementsDataStorage/Source/TypedElementsDataStorage/Private/TypedElementDatabaseEnvironment.cpp
+++ b/Engine/Plugins/Experimental/TypedElementsDataStorage/Source/TypedElementsDataStorage/Private/TypedElementDatabaseEnvironment.cpp
@@ -2,9 +2,10 @@
 
 #include "TypedElementDatabaseEnvironment.h"
 
-FTypedElementDatabaseEnvironment::FTypedElementDatabaseEnvironment(
+FTypedElementDatabaseEnvironment::FTypedElementDatabaseEnvironment(ITypedElementDataStorageInterface& DataStorage,
 	FMassEntityManager& InMassEntityManager, FMassProcessingPhaseManager& InMassPhaseManager)
 	: DirectDeferredCommands(*this)
+	, MementoSystem(DataStorage)
 	, MassEntityManager(InMassEntityManager)
 	, MassPhaseManager(InMassPhaseManager)
 {
@@ -50,6 +51,16 @@ const FTypedElementExtendedQueryStore& FTypedElementDatabaseEnvironment::GetQuer
 	return Queries;
 }
 
+UTypedElementMementoSystem& FTypedElementDatabaseEnvironment::GetMementoSystem()
+{
+	return MementoSystem;
+}
+
+const UTypedElementMementoSystem& FTypedElementDatabaseEnvironment::GetMementoSystem() const
+{
+	return MementoSystem;
+}
+
 FMassEntityManager& FTypedElementDatabaseEnvironment::GetMassEntityManager()
 {
 	return MassEntityManager;
diff --git a/Engine/Plugins/Experimental/TypedElementsDataStorage/Source/TypedElementsDataStorage/Private/TypedElementDatabaseEnvironment.h b/Engine/Plugins/Experimental/TypedElementsDataStorage/Source/TypedElementsDataStorage/Private/TypedElementDatabaseEnvironment.h
--- a/Engine/Plugins/Experimental/TypedElementsDataStorage/Source/TypedElementsDataStorage/Private/TypedElementDatabaseEnvironment.h
+++ b/Engine/Plugins/Experimental/TypedElementsDataStorage/Source/TypedElementsDataStorage/Private/TypedElementDatabaseEnvironment.h
@@ -7,12 +7,14 @@
 #include "TypedElementDatabaseCommandBuffer.h"
 #include "TypedElementDatabaseScratchBuffer.h"
 #include "TypedElementDatabaseIndexTable.h"
+#include "Memento/TypedElementMementoSystem.h"
 #include "Queries/TypedElementExtendedQueryStore.h"
 
 class FTypedElementDatabaseEnvironment final
 {
 public:
-	FTypedElementDatabaseEnvironment(FMassEntityManager& InMassEntityManager, FMassProcessingPhaseManager& InMassPhaseManager);
+	FTypedElementDatabaseEnvironment(ITypedElementDataStorageInterface& DataStorage, 
+		FMassEntityManager& InMassEntityManager, FMassProcessingPhaseManager& InMassPhaseManager);
 
 	FTypedElementDatabaseCommandBuffer& GetDirectDeferredCommands();
 	const FTypedElementDatabaseCommandBuffer& GetDirectDeferredCommands() const;
@@ -26,6 +28,9 @@ public:
 	FTypedElementExtendedQueryStore& GetQueryStore();
 	const FTypedElementExtendedQueryStore& GetQueryStore() const;
 
+	UTypedElementMementoSystem& GetMementoSystem();
+	const UTypedElementMementoSystem& GetMementoSystem() const;
+
 	FMassEntityManager& GetMassEntityManager();
 	const FMassEntityManager& GetMassEntityManager() const;
 	
@@ -40,7 +45,8 @@ private:
 	FTypedElementDatabaseIndexTable IndexTable;
 	FTypedElementDatabaseScratchBuffer ScratchBuffer;
 	FTypedElementExtendedQueryStore Queries;
-	
+	UTypedElementMementoSystem MementoSystem;
+
 	FMassEntityManager& MassEntityManager;
 	FMassProcessingPhaseManager& MassPhaseManager;
 
diff --git a/Engine/Plugins/Experimental/TypedElementsDataStorage/Source/TypedElementsDataStorage/Private/TypedElementsDataStorage.cpp b/Engine/Plugins/Experimental/TypedElementsDataStorage/Source/TypedElementsDataStorage/Private/TypedElementsDataStorage.cpp
--- a/Engine/Plugins/Experimental/TypedElementsDataStorage/Source/TypedElementsDataStorage/Private/TypedElementsDataStorage.cpp
+++ b/Engine/Plugins/Experimental/TypedElementsDataStorage/Source/TypedElementsDataStorage/Private/TypedElementsDataStorage.cpp
@@ -91,11 +91,8 @@ void FTypedElementsDataStorageModule::StartupModule()
 				DatabaseUi = NewObject<UTypedElementDatabaseUi>();
 				DatabaseUi->Initialize(Database.Get(), DatabaseCompatibility.Get());
 
-				MementoSystem = NewObject<UTypedElementMementoSystem>();
-				MementoSystem->Initialize(*Database.Get());
-
 				ObjectReinstancingManager = NewObject<UTypedElementObjectReinstancingManager>();
-				ObjectReinstancingManager->Initialize(*Database, *DatabaseCompatibility, *MementoSystem);
+				ObjectReinstancingManager->Initialize(*Database, *DatabaseCompatibility);
 
 				// Register the various database instances.
 				UTypedElementRegistry* Registry = UTypedElementRegistry::GetInstance();
@@ -162,7 +159,6 @@ void FTypedElementsDataStorageModule::ShutdownModule()
 		if (UObjectInitialized())
 		{
 			ObjectReinstancingManager->Deinitialize();
-			MementoSystem->Deinitialize();
 			DatabaseUi->Deinitialize();
 			DatabaseCompatibility->Deinitialize();
 			Database->Deinitialize();
@@ -179,7 +175,6 @@ void FTypedElementsDataStorageModule::AddReferencedObjects(FReferenceCollector&
 		Collector.AddReferencedObject(Database);
 		Collector.AddReferencedObject(DatabaseCompatibility);
 		Collector.AddReferencedObject(DatabaseUi);
-		Collector.AddReferencedObject(MementoSystem);
 		Collector.AddReferencedObject(ObjectReinstancingManager);
 	}
 }
diff --git a/Engine/Plugins/Experimental/TypedElementsDataStorage/Source/TypedElementsDataStorage/Public/Memento/TypedElementMementoInterface.h b/Engine/Plugins/Experimental/TypedElementsDataStorage/Source/TypedElementsDataStorage/Public/Memento/TypedElementMementoInterface.h
deleted file mode 100644
diff --git a/Engine/Plugins/Experimental/TypedElementsDataStorage/Source/TypedElementsDataStorage/Public/TypedElementsDataStorage.h b/Engine/Plugins/Experimental/TypedElementsDataStorage/Source/TypedElementsDataStorage/Public/TypedElementsDataStorage.h
--- a/Engine/Plugins/Experimental/TypedElementsDataStorage/Source/TypedElementsDataStorage/Public/TypedElementsDataStorage.h
+++ b/Engine/Plugins/Experimental/TypedElementsDataStorage/Source/TypedElementsDataStorage/Public/TypedElementsDataStorage.h
@@ -12,7 +12,6 @@ class UTypedElementDatabase;
 class UTypedElementDatabaseCompatibility;
 class UTypedElementDatabaseUi;
 class UTypedElementObjectReinstancingManager;
-class UTypedElementMementoSystem;
 
 class FTypedElementsDataStorageModule : public IModuleInterface, public FGCObject
 {
@@ -38,7 +37,6 @@ private:
 	TObjectPtr<UTypedElementDatabase> Database;
 	TObjectPtr<UTypedElementDatabaseCompatibility> DatabaseCompatibility;
 	TObjectPtr<UTypedElementDatabaseUi> DatabaseUi;
-	TObjectPtr<UTypedElementMementoSystem> MementoSystem;
 	TObjectPtr<UTypedElementObjectReinstancingManager> ObjectReinstancingManager;
 	bool bInitialized{ false };
 };
diff --git a/Engine/Source/Runtime/TypedElementFramework/Public/Elements/Interfaces/TypedElementDataStorageInterface.h b/Engine/Source/Runtime/TypedElementFramework/Public/Elements/Interfaces/TypedElementDataStorageInterface.h
--- a/Engine/Source/Runtime/TypedElementFramework/Public/Elements/Interfaces/TypedElementDataStorageInterface.h
+++ b/Engine/Source/Runtime/TypedElementFramework/Public/Elements/Interfaces/TypedElementDataStorageInterface.h
@@ -88,16 +88,20 @@ public:
 	 */
 
 	/** Creates a new table for with the provided columns. Optionally a name can be given which is useful for retrieval later. */
-	virtual TypedElementTableHandle RegisterTable(TConstArrayView<const UScriptStruct*> ColumnList, const FName Name) = 0;
+	virtual TypedElementDataStorage::TableHandle RegisterTable(TConstArrayView<const UScriptStruct*> ColumnList, const FName Name) = 0;
+	template<TypedElementDataStorage::TColumnType... Columns>
+	TypedElementDataStorage::TableHandle RegisterTable(const FName Name);
 	/** 
 	 * Copies the column information from the provided table and creates a new table for with the provided columns. Optionally a 
 	 * name can be given which is useful for retrieval later.
 	 */
-	virtual TypedElementTableHandle RegisterTable(TypedElementTableHandle SourceTable, 
+	virtual TypedElementDataStorage::TableHandle RegisterTable(TypedElementDataStorage::TableHandle SourceTable,
 		TConstArrayView<const UScriptStruct*> ColumnList, const FName Name) = 0;
+	template<TypedElementDataStorage::TColumnType... Columns>
+	TypedElementDataStorage::TableHandle RegisterTable(TypedElementDataStorage::TableHandle SourceTable, const FName Name);
 
 	/** Returns a previously created table with the provided name or TypedElementInvalidTableHandle if not found. */
-	virtual TypedElementTableHandle FindTable(const FName Name) = 0;
+	virtual TypedElementDataStorage::TableHandle FindTable(const FName Name) = 0;
 	
 	/**
 	 * @section Row management
@@ -312,6 +316,11 @@ public:
 	 * for the data storage.
 	 */
 	virtual FTypedElementOnDataStorageUpdate& OnUpdate() = 0;
+	/**
+	 * Called periodically when the storage is available. This provides an opportunity clean up after processing and
+	 * to get ready for the next batch up updates.
+	 */
+	virtual FTypedElementOnDataStorageUpdate& OnUpdateCompleted() = 0;
 
 	/**
 	 * Whether or not the data storage is available. The data storage is available most of the time, but can be
@@ -337,6 +346,19 @@ const FactoryT* ITypedElementDataStorageInterface::FindFactory() const
 	return static_cast<const FactoryT*>(FindFactory(FactoryT::StaticClass()));
 }
 
+template<TypedElementDataStorage::TColumnType... Columns>
+TypedElementDataStorage::TableHandle ITypedElementDataStorageInterface::RegisterTable(const FName Name)
+{
+	return RegisterTable({ Columns::StaticStruct()... }, Name);
+}
+
+template<TypedElementDataStorage::TColumnType... Columns>
+TypedElementDataStorage::TableHandle ITypedElementDataStorageInterface::RegisterTable(
+	TypedElementDataStorage::TableHandle SourceTable, const FName Name)
+{
+	return RegisterTable(SourceTable, { Columns::StaticStruct()... }, Name);
+}
+
 template<TypedElementDataStorage::TColumnType Column>
 void ITypedElementDataStorageInterface::AddColumn(TypedElementRowHandle Row)
 {
