Mirror the World in UObjects to TEDS.

This update stores a weak pointer to the World for any UObjects stored in TEDS in a dedicated column and removes it if it's no longer present.

diff --git a/Engine/Plugins/Experimental/TypedElementsDataStorage/Source/TypedElementsDataStorage/Private/Compatibility/TypedElementUObjectWorldQueries.cpp b/Engine/Plugins/Experimental/TypedElementsDataStorage/Source/TypedElementsDataStorage/Private/Compatibility/TypedElementUObjectWorldQueries.cpp
new file mode 100644
--- /dev/null
+++ b/Engine/Plugins/Experimental/TypedElementsDataStorage/Source/TypedElementsDataStorage/Private/Compatibility/TypedElementUObjectWorldQueries.cpp
@@ -0,0 +1,75 @@
+// Copyright Epic Games, Inc. All Rights Reserved.
+
+#include "Compatibility/TypedElementUObjectWorldQueries.h"
+
+#include "Elements/Columns/TypedElementCompatibilityColumns.h"
+#include "Elements/Columns/TypedElementMiscColumns.h"
+#include "Elements/Framework/TypedElementQueryBuilder.h"
+#include "Engine/World.h"
+
+void UTypedElementUObjectWorldFactory::RegisterQueries(ITypedElementDataStorageInterface& DataStorage)
+{
+	RegisterAddWorldColumn(DataStorage);
+	RegisterUpdateOrRemoveWorldColumn(DataStorage);
+}
+
+void UTypedElementUObjectWorldFactory::RegisterAddWorldColumn(ITypedElementDataStorageInterface& DataStorage) const
+{
+	using namespace TypedElementDataStorage;
+	using namespace TypedElementQueryBuilder;
+	
+	DataStorage.RegisterQuery(
+		Select(
+			TEXT("Add world column to UObject"),
+			FProcessor(EQueryTickPhase::PrePhysics, DataStorage.GetQueryTickGroupName(EQueryTickGroups::SyncExternalToDataStorage))
+			.ForceToGameThread(true),
+			[](IQueryContext& Context, TypedElementRowHandle Row, const FTypedElementUObjectColumn& Object)
+			{
+				// Not all objects, in particular actors, are always correctly cleaned up, resulting in dangling
+				// pointers in TEDS.
+				if (UObject* ObjectInstance = Object.Object.Get())
+				{
+					if (UWorld* World = ObjectInstance->GetWorld())
+					{
+						Context.AddColumn(Row, FTypedElementWorldColumn{ .World = World });
+					}
+				}
+			}
+		)
+		.Where()
+			.All<FTypedElementSyncFromWorldTag>()
+			.None<FTypedElementWorldColumn>()
+		.Compile()
+	);
+}
+
+void UTypedElementUObjectWorldFactory::RegisterUpdateOrRemoveWorldColumn(ITypedElementDataStorageInterface& DataStorage) const
+{
+	using namespace TypedElementDataStorage;
+	using namespace TypedElementQueryBuilder;
+	
+	DataStorage.RegisterQuery(
+		Select(
+			TEXT("Sync UObject's world to column"),
+			FProcessor(EQueryTickPhase::PostPhysics, DataStorage.GetQueryTickGroupName(EQueryTickGroups::SyncExternalToDataStorage))
+			.ForceToGameThread(true),
+			[](IQueryContext& Context, TypedElementRowHandle Row, const FTypedElementUObjectColumn& Object, FTypedElementWorldColumn& World)
+			{
+				// Not all objects, in particular actors, are always correctly cleaned up, resulting in dangling
+				// pointers in TEDS.
+				if (UObject* ObjectInstance = Object.Object.Get())
+				{
+					if (UWorld* UObjectWorld = ObjectInstance->GetWorld())
+					{
+						World.World = UObjectWorld;
+						return;
+					}
+				}
+				Context.RemoveColumns<FTypedElementWorldColumn>(Row);
+			}
+		)
+		.Where()
+			.All<FTypedElementSyncFromWorldTag>()
+		.Compile()
+	);
+}
diff --git a/Engine/Plugins/Experimental/TypedElementsDataStorage/Source/TypedElementsDataStorage/Private/Compatibility/TypedElementUObjectWorldQueries.h b/Engine/Plugins/Experimental/TypedElementsDataStorage/Source/TypedElementsDataStorage/Private/Compatibility/TypedElementUObjectWorldQueries.h
new file mode 100644
--- /dev/null
+++ b/Engine/Plugins/Experimental/TypedElementsDataStorage/Source/TypedElementsDataStorage/Private/Compatibility/TypedElementUObjectWorldQueries.h
@@ -0,0 +1,31 @@
+// Copyright Epic Games, Inc. All Rights Reserved.
+
+#pragma once
+
+#include "Elements/Interfaces/TypedElementDataStorageFactory.h"
+#include "UObject/ObjectMacros.h"
+
+#include "TypedElementUObjectWorldQueries.generated.h"
+
+UCLASS()
+class UTypedElementUObjectWorldFactory : public UTypedElementDataStorageFactory
+{
+	GENERATED_BODY()
+
+public:
+	~UTypedElementUObjectWorldFactory() override = default;
+
+	void RegisterQueries(ITypedElementDataStorageInterface& DataStorage) override;
+
+private:
+	/**
+	 * Checks rows with UObjects that don't have a world column yet if one needs to be added whenever
+	 * the row is marked for updates.
+	 */
+	void RegisterAddWorldColumn(ITypedElementDataStorageInterface& DataStorage) const;
+	/**
+	 * Updates the world column with the world in the UObject or removes it if there's no world associated
+	 * with the UObject anymore.
+	 */
+	void RegisterUpdateOrRemoveWorldColumn(ITypedElementDataStorageInterface& DataStorage) const;
+};
diff --git a/Engine/Source/Runtime/TypedElementFramework/Public/Elements/Columns/TypedElementCompatibilityColumns.h b/Engine/Source/Runtime/TypedElementFramework/Public/Elements/Columns/TypedElementCompatibilityColumns.h
--- a/Engine/Source/Runtime/TypedElementFramework/Public/Elements/Columns/TypedElementCompatibilityColumns.h
+++ b/Engine/Source/Runtime/TypedElementFramework/Public/Elements/Columns/TypedElementCompatibilityColumns.h
@@ -10,6 +10,7 @@
 #include "TypedElementCompatibilityColumns.generated.h"
 
 class UObject;
+class UWorld;
 
 /**
  * Column containing a non-owning reference to a UObject.
@@ -44,3 +45,14 @@ struct FTypedElementActorTag final : public FTypedElementDataStorageTag
 {
 	GENERATED_BODY()
 };
+
+/**
+ * Column that stores a reference to the world.
+ */
+USTRUCT(meta = (DisplayName = "World"))
+struct FTypedElementWorldColumn final : public FTypedElementDataStorageColumn
+{
+	GENERATED_BODY()
+
+	TWeakObjectPtr<UWorld> World;
+};
