Add parent column for actors in TEDS.

This change adds support for adding/removing parent columns to actors if an actor has a parent. If an actor is reparented, the parent column will be updated.

diff --git a/Engine/Plugins/Experimental/TypedElementsDataStorage/Source/TypedElementsDataStorage/Private/Compatibility/TypedElementActorParentQueries.cpp b/Engine/Plugins/Experimental/TypedElementsDataStorage/Source/TypedElementsDataStorage/Private/Compatibility/TypedElementActorParentQueries.cpp
new file mode 100644
--- /dev/null
+++ b/Engine/Plugins/Experimental/TypedElementsDataStorage/Source/TypedElementsDataStorage/Private/Compatibility/TypedElementActorParentQueries.cpp
@@ -0,0 +1,95 @@
+// Copyright Epic Games, Inc. All Rights Reserved.
+
+#include "Compatibility/TypedElementActorParentQueries.h"
+
+#include "Elements/Columns/TypedElementCompatibilityColumns.h"
+#include "Elements/Columns/TypedElementHiearchyColumns.h"
+#include "Elements/Columns/TypedElementMiscColumns.h"
+#include "Elements/Framework/TypedElementIndexHasher.h"
+#include "Elements/Framework/TypedElementQueryBuilder.h"
+#include "GameFramework/Actor.h"
+
+void UTypedElementActorParentFactory::RegisterQueries(ITypedElementDataStorageInterface& DataStorage)
+{
+	RegisterAddParentColumn(DataStorage);
+	RegisterUpdateOrRemoveParentColumn(DataStorage);
+}
+
+void UTypedElementActorParentFactory::RegisterAddParentColumn(ITypedElementDataStorageInterface& DataStorage) const
+{
+	using namespace TypedElementDataStorage;
+	using namespace TypedElementQueryBuilder;
+	
+	DataStorage.RegisterQuery(
+		Select(
+			TEXT("Add parent column to actor"),
+			FProcessor(EQueryTickPhase::PrePhysics, DataStorage.GetQueryTickGroupName(EQueryTickGroups::SyncExternalToDataStorage))
+			.ForceToGameThread(true),
+			[](IQueryContext& Context, TypedElementRowHandle Row, const FTypedElementUObjectColumn& Actor)
+			{
+				if (const AActor* ActorInstance = Cast<AActor>(Actor.Object))
+				{
+					if (AActor* Parent = ActorInstance->GetAttachParentActor())
+					{
+						uint64 IdHash = GenerateIndexHash(Parent);
+						RowHandle ParentRow = Context.FindIndexedRow(IdHash);
+						if (Context.IsRowAvailable(ParentRow))
+						{
+							Context.AddColumn(Row, FTypedElementParentColumn{ .Parent = ParentRow });
+						}
+						else
+						{
+							Context.AddColumn(Row, FTypedElementUnresolvedParentColumn{ .ParentIdHash = IdHash });
+						}
+					}
+				}
+			}
+		)
+		.Where()
+			.All<FTypedElementSyncFromWorldTag, FTypedElementActorTag>()
+			.None<FTypedElementParentColumn, FTypedElementUnresolvedParentColumn>()
+		.Compile()
+	);
+}
+
+void UTypedElementActorParentFactory::RegisterUpdateOrRemoveParentColumn(ITypedElementDataStorageInterface& DataStorage) const
+{
+	using namespace TypedElementDataStorage;
+	using namespace TypedElementQueryBuilder;
+	
+	DataStorage.RegisterQuery(
+		Select(
+			TEXT("Sync actor's parent to column"),
+			FProcessor(EQueryTickPhase::PostPhysics, DataStorage.GetQueryTickGroupName(EQueryTickGroups::SyncExternalToDataStorage))
+				.ForceToGameThread(true),
+			[](IQueryContext& Context, TypedElementRowHandle Row, const FTypedElementUObjectColumn& Actor, FTypedElementParentColumn& Parent)
+			{
+				if (const AActor* ActorInstance = Cast<AActor>(Actor.Object))
+				{
+					if (AActor* ParentActor = ActorInstance->GetAttachParentActor())
+					{
+						uint64 IdHash = GenerateIndexHash(ParentActor);
+						if (Parent.Parent != IdHash)
+						{
+							RowHandle ParentRow = Context.FindIndexedRow(IdHash);
+							if (Context.IsRowAvailable(ParentRow))
+							{
+								Parent.Parent = ParentRow;
+							}
+							else
+							{
+								Context.RemoveColumns<FTypedElementParentColumn>(Row);
+								Context.AddColumn(Row, FTypedElementUnresolvedParentColumn{ .ParentIdHash = IdHash });
+							}
+						}
+						return;
+					}
+				}
+				Context.RemoveColumns<FTypedElementParentColumn>(Row);
+			}
+		)
+		.Where()
+			.All<FTypedElementActorTag, FTypedElementSyncFromWorldTag>()
+		.Compile()
+	);
+}
diff --git a/Engine/Plugins/Experimental/TypedElementsDataStorage/Source/TypedElementsDataStorage/Private/Compatibility/TypedElementActorParentQueries.h b/Engine/Plugins/Experimental/TypedElementsDataStorage/Source/TypedElementsDataStorage/Private/Compatibility/TypedElementActorParentQueries.h
new file mode 100644
--- /dev/null
+++ b/Engine/Plugins/Experimental/TypedElementsDataStorage/Source/TypedElementsDataStorage/Private/Compatibility/TypedElementActorParentQueries.h
@@ -0,0 +1,31 @@
+// Copyright Epic Games, Inc. All Rights Reserved.
+
+#pragma once
+
+#include "Elements/Interfaces/TypedElementDataStorageFactory.h"
+#include "UObject/ObjectMacros.h"
+
+#include "TypedElementActorParentQueries.generated.h"
+
+UCLASS()
+class UTypedElementActorParentFactory : public UTypedElementDataStorageFactory
+{
+	GENERATED_BODY()
+
+public:
+	~UTypedElementActorParentFactory() override = default;
+
+	void RegisterQueries(ITypedElementDataStorageInterface& DataStorage) override;
+
+private:
+	/**
+	 * Checks rows with actors that don't have a parent column yet if one needs to be added whenever
+	 * the row is marked for updates.
+	 */
+	void RegisterAddParentColumn(ITypedElementDataStorageInterface& DataStorage) const;
+	/**
+	 * Updates the parent column with the parent from the actor or removes it if there's no parent associated
+	 * with the actor anymore.
+	 */
+	void RegisterUpdateOrRemoveParentColumn(ITypedElementDataStorageInterface& DataStorage) const;
+};
