TEDS - Disable adding parent information for actors by default and instead hide it behind a Cvar

diff --git a/Engine/Plugins/Experimental/TypedElementsDataStorage/Source/TypedElementsDataStorage/Private/Compatibility/TypedElementActorParentQueries.cpp b/Engine/Plugins/Experimental/TypedElementsDataStorage/Source/TypedElementsDataStorage/Private/Compatibility/TypedElementActorParentQueries.cpp
--- a/Engine/Plugins/Experimental/TypedElementsDataStorage/Source/TypedElementsDataStorage/Private/Compatibility/TypedElementActorParentQueries.cpp
+++ b/Engine/Plugins/Experimental/TypedElementsDataStorage/Source/TypedElementsDataStorage/Private/Compatibility/TypedElementActorParentQueries.cpp
@@ -9,10 +9,24 @@
 #include "Elements/Framework/TypedElementQueryBuilder.h"
 #include "GameFramework/Actor.h"
 
+namespace UE::TypedElementActorParentQueries::Local
+{
+	static bool bAddParentColumnToActors = false;
+	
+	// Cvar to allow the TEDS-Outliner to automatically take over the level editor's 4 Outliners instead of appearing as a separate tab
+	static FAutoConsoleVariableRef CVarUseTEDSOutliner(
+		TEXT("TEDS.AddParentColumnToActors"),
+		bAddParentColumnToActors,
+		TEXT("Mirror parent information for actors to TEDS (only works when set on startup)"));
+};
+
 void UTypedElementActorParentFactory::RegisterQueries(ITypedElementDataStorageInterface& DataStorage)
 {
-	RegisterAddParentColumn(DataStorage);
-	RegisterUpdateOrRemoveParentColumn(DataStorage);
+	if(UE::TypedElementActorParentQueries::Local::bAddParentColumnToActors)
+	{
+		RegisterAddParentColumn(DataStorage);
+		RegisterUpdateOrRemoveParentColumn(DataStorage);
+	}
 }
 
 void UTypedElementActorParentFactory::RegisterAddParentColumn(ITypedElementDataStorageInterface& DataStorage) const
