#UE Readd MassActorSubsystem.h header to fix CIS

diff --git a/Engine/Plugins/Experimental/TypedElementsDataStorage/Source/TypedElementsDataStorageRevisionControl/Source/Private/RevisionControlProcessors.cpp b/Engine/Plugins/Experimental/TypedElementsDataStorage/Source/TypedElementsDataStorageRevisionControl/Source/Private/RevisionControlProcessors.cpp
--- a/Engine/Plugins/Experimental/TypedElementsDataStorage/Source/TypedElementsDataStorageRevisionControl/Source/Private/RevisionControlProcessors.cpp
+++ b/Engine/Plugins/Experimental/TypedElementsDataStorage/Source/TypedElementsDataStorageRevisionControl/Source/Private/RevisionControlProcessors.cpp
@@ -14,6 +14,8 @@
 #include "Elements/Framework/TypedElementQueryBuilder.h"
 #include "Elements/Framework/TypedElementRegistry.h"
 
+#include "MassActorSubsystem.h"
+
 extern TYPEDELEMENTSDATASTORAGE_API FAutoConsoleVariableRef CVarAutoPopulateState;
 
 static bool gEnableOverlays = false;
