Changed default opacity value from 85% to 20% for revision control overlay color

diff --git a/Engine/Plugins/Experimental/TypedElementsDataStorage/Source/TypedElementsDataStorageRevisionControl/Source/Private/RevisionControlProcessors.cpp b/Engine/Plugins/Experimental/TypedElementsDataStorage/Source/TypedElementsDataStorageRevisionControl/Source/Private/RevisionControlProcessors.cpp
--- a/Engine/Plugins/Experimental/TypedElementsDataStorage/Source/TypedElementsDataStorageRevisionControl/Source/Private/RevisionControlProcessors.cpp
+++ b/Engine/Plugins/Experimental/TypedElementsDataStorage/Source/TypedElementsDataStorageRevisionControl/Source/Private/RevisionControlProcessors.cpp
@@ -51,7 +51,7 @@ TAutoConsoleVariable<bool> CVarEnableOverlayOpenForAdd(
 	TEXT("Enables overlays for files that are newly added."),
 	ECVF_Default);
 
-static int32 gOverlayAlpha = 217;
+static int32 gOverlayAlpha = 51;
 TAutoConsoleVariable<int32> CVarOverlayAlpha(
 	TEXT("SourceControl.Overlays.Alpha"),
 	gOverlayAlpha,
