Skein / ViewportStatus: Fixed opacity values losing precision during conversion. The value is now stored as-is and the conversion to color space only happens on use of the value.

diff --git a/Engine/Plugins/Experimental/TypedElementsDataStorage/Source/TypedElementsDataStorageRevisionControl/Source/Private/RevisionControlProcessors.cpp b/Engine/Plugins/Experimental/TypedElementsDataStorage/Source/TypedElementsDataStorageRevisionControl/Source/Private/RevisionControlProcessors.cpp
--- a/Engine/Plugins/Experimental/TypedElementsDataStorage/Source/TypedElementsDataStorageRevisionControl/Source/Private/RevisionControlProcessors.cpp
+++ b/Engine/Plugins/Experimental/TypedElementsDataStorage/Source/TypedElementsDataStorageRevisionControl/Source/Private/RevisionControlProcessors.cpp
@@ -53,7 +53,7 @@ TAutoConsoleVariable<bool> CVarEnableOverlayOpenForAdd(
 	TEXT("Enables overlays for files that are newly added."),
 	ECVF_Default);
 
-static int32 gOverlayAlpha = 51;
+static int32 gOverlayAlpha = 20; // [0..100]
 TAutoConsoleVariable<int32> CVarOverlayAlpha(
 	TEXT("SourceControl.Overlays.Alpha"),
 	gOverlayAlpha,
@@ -69,13 +69,16 @@ static FColor DetermineOverlayColor(const TypedElementDataStorage::IQueryContext
 	bool bSelected = ObjectContext.HasColumn<FTypedElementSelectionColumn>();
 	if (!bIgnored && !bSelected)
 	{
+		// Convert CVar value from [0..100] to [0..255] range.
+		int32 Alpha = FMath::Lerp<float>(0.f, 255.f, CVarOverlayAlpha.GetValueOnGameThread() / 100.f);
+
 		// Check if the package is outdated because there is a newer version available.
 		if (SCCContext.HasColumn<FSCCNotCurrentTag>())
 		{
 			if (CVarEnableOverlayNotAtHeadRevision.GetValueOnGameThread())
 			{
 				// Yellow.
-				return FColor(225, 255, 61, CVarOverlayAlpha.GetValueOnGameThread());
+				return FColor(225, 255, 61, Alpha);
 			}
 		}
 
@@ -85,7 +88,7 @@ static FColor DetermineOverlayColor(const TypedElementDataStorage::IQueryContext
 			if (CVarEnableOverlayCheckedOutByOtherUser.GetValueOnGameThread())
 			{
 				// Red.
-				return FColor(239, 53, 53, CVarOverlayAlpha.GetValueOnGameThread());
+				return FColor(239, 53, 53, Alpha);
 			}
 		}
 
@@ -99,7 +102,7 @@ static FColor DetermineOverlayColor(const TypedElementDataStorage::IQueryContext
 					if (StatusColumn->Modification == ESCCModification::Added)
 					{
 						// Blue.
-						return FColor(0, 112, 224, CVarOverlayAlpha.GetValueOnGameThread());
+						return FColor(0, 112, 224, Alpha);
 					}
 				}
 			}
@@ -111,7 +114,7 @@ static FColor DetermineOverlayColor(const TypedElementDataStorage::IQueryContext
 			if (CVarEnableOverlayCheckedOut.GetValueOnGameThread())
 			{
 				// Green.
-				return FColor(31, 228, 75, CVarOverlayAlpha.GetValueOnGameThread());
+				return FColor(31, 228, 75, Alpha);
 			}
 		}
 	}
diff --git a/Engine/Source/Developer/SourceControlViewport/Private/SourceControlViewportUtils.cpp b/Engine/Source/Developer/SourceControlViewport/Private/SourceControlViewportUtils.cpp
--- a/Engine/Source/Developer/SourceControlViewport/Private/SourceControlViewportUtils.cpp
+++ b/Engine/Source/Developer/SourceControlViewport/Private/SourceControlViewportUtils.cpp
@@ -79,20 +79,18 @@ void SetFeedbackEnabled(FViewportClient* ViewportClient, ESourceControlStatus St
 
 uint8 GetFeedbackOpacity(FViewportClient* ViewportClient)
 {
-	uint8 Opacity = 128;
+	uint8 Opacity = 0;
 
 	if (IConsoleVariable* CVar = IConsoleManager::Get().FindConsoleVariable(TEXT("SourceControl.Overlays.Alpha")); ensure(CVar))
 	{
 		Opacity = CVar->GetInt();
 	}
 
-	return FMath::Lerp<float>(0.f, 100.f, Opacity / 255.f);
+	return Opacity;
 }
 
 void SetFeedbackOpacity(FViewportClient* ViewportClient, uint8 Opacity)
 {
-	Opacity = FMath::Lerp<float>(0.f, 255.f, Opacity / 100.f);
-
 	if (IConsoleVariable* CVar = IConsoleManager::Get().FindConsoleVariable(TEXT("SourceControl.Overlays.Alpha")); ensure(CVar))
 	{
 		CVar->Set(Opacity);
