Skein / ViewportStatus: Fixed overlays remaining active after disabling revision control.

The OnChanged() delegate is somewhat misleading. It doesn't just fire when the CVar is changed, it fires whenever the CVar is set - even if it's set to the same value. This means that previously, several 'ApplyOverlays' queries would be registered simultaneously as the CVar gets sets several times, which could also be a cause of hitches. It also meant that the overlays kept being applied, even after revision control got disabled, because not all queries got unregistered.

diff --git a/Engine/Plugins/Experimental/TypedElementsDataStorage/Source/TypedElementsDataStorageRevisionControl/Source/Private/RevisionControlProcessors.cpp b/Engine/Plugins/Experimental/TypedElementsDataStorage/Source/TypedElementsDataStorageRevisionControl/Source/Private/RevisionControlProcessors.cpp
--- a/Engine/Plugins/Experimental/TypedElementsDataStorage/Source/TypedElementsDataStorageRevisionControl/Source/Private/RevisionControlProcessors.cpp
+++ b/Engine/Plugins/Experimental/TypedElementsDataStorage/Source/TypedElementsDataStorageRevisionControl/Source/Private/RevisionControlProcessors.cpp
@@ -133,7 +133,6 @@ void UTypedElementRevisionControlFactory::RegisterTables(ITypedElementDataStorag
 
 void UTypedElementRevisionControlFactory::RegisterQueries(ITypedElementDataStorageInterface& DataStorage)
 {
-
 	CVarAutoPopulateState->AsVariable()->OnChangedDelegate().AddLambda(
 		[this, &DataStorage](IConsoleVariable* AutoPopulate)
 		{
@@ -144,6 +143,7 @@ void UTypedElementRevisionControlFactory::RegisterQueries(ITypedElementDataStora
 			else
 			{
 				DataStorage.UnregisterQuery(FetchUpdates);
+				FetchUpdates = TypedElementInvalidQueryHandle;
 			}
 		}
 	);
@@ -153,12 +153,20 @@ void UTypedElementRevisionControlFactory::RegisterQueries(ITypedElementDataStora
 		{
 			if (EnableOverlays->GetBool())
 			{
+				DataStorage.UnregisterQuery(RemoveOverlays);
+				RemoveOverlays = TypedElementInvalidQueryHandle;
+
 				RegisterApplyOverlays(DataStorage);
 			}
 			else
 			{
 				DataStorage.UnregisterQuery(ApplyOverlays);
+				ApplyOverlays = TypedElementInvalidQueryHandle;
+
 				DataStorage.UnregisterQuery(ApplyOverlaysObjectToSCC);
+				ApplyOverlaysObjectToSCC = TypedElementInvalidQueryHandle;
+
+				RegisterRemoveOverlays(DataStorage);
 			}
 		}
 	);
@@ -172,6 +180,10 @@ void UTypedElementRevisionControlFactory::RegisterQueries(ITypedElementDataStora
 	{
 		RegisterApplyOverlays(DataStorage);
 	}
+	else
+	{
+		RegisterRemoveOverlays(DataStorage);
+	}
 }
 
 void UTypedElementRevisionControlFactory::RegisterFetchUpdates(ITypedElementDataStorageInterface& DataStorage) const
@@ -181,36 +193,39 @@ void UTypedElementRevisionControlFactory::RegisterFetchUpdates(ITypedElementData
 	
 	FSourceControlFileStatusMonitor& FileStatusMonitor = ISourceControlModule::Get().GetSourceControlFileStatusMonitor();
 
-	FetchUpdates = DataStorage.RegisterQuery(
-		Select(
-			TEXT("Gather source control statuses for objects with unresolved package paths"),
-			FProcessor(DSI::EQueryTickPhase::DuringPhysics, DataStorage.GetQueryTickGroupName(DSI::EQueryTickGroups::SyncExternalToDataStorage))
-				.ForceToGameThread(true),
-			[this, &FileStatusMonitor](DSI::IQueryContext& Context, const FTypedElementPackageUnresolvedReference* InUnresolvedReferences)
-			{
-				TConstArrayView<TypedElementDataStorage::RowHandle> RowHandles = Context.GetRowHandles();
-				TConstArrayView<FTypedElementPackageUnresolvedReference, int64> UnresolvedReferences { InUnresolvedReferences, Context.GetRowCount() };
-
-				for (int64 UnresolvedReferenceIndex = 0; UnresolvedReferenceIndex < UnresolvedReferences.Num(); ++UnresolvedReferenceIndex)
+	if (FetchUpdates == TypedElementInvalidQueryHandle)
+	{
+		FetchUpdates = DataStorage.RegisterQuery(
+			Select(
+				TEXT("Gather source control statuses for objects with unresolved package paths"),
+				FProcessor(DSI::EQueryTickPhase::DuringPhysics, DataStorage.GetQueryTickGroupName(DSI::EQueryTickGroups::SyncExternalToDataStorage))
+					.ForceToGameThread(true),
+				[this, &FileStatusMonitor](DSI::IQueryContext& Context, const FTypedElementPackageUnresolvedReference* InUnresolvedReferences)
 				{
-					const FTypedElementPackageUnresolvedReference& UnresolvedReference = UnresolvedReferences[UnresolvedReferenceIndex];
-					if (UnresolvedReference.Index == 0)
+					TConstArrayView<TypedElementDataStorage::RowHandle> RowHandles = Context.GetRowHandles();
+					TConstArrayView<FTypedElementPackageUnresolvedReference, int64> UnresolvedReferences { InUnresolvedReferences, Context.GetRowCount() };
+
+					for (int64 UnresolvedReferenceIndex = 0; UnresolvedReferenceIndex < UnresolvedReferences.Num(); ++UnresolvedReferenceIndex)
 					{
-						Context.RemoveColumns<FTypedElementPackageUnresolvedReference>(RowHandles[UnresolvedReferenceIndex]);
-						return;
-					}
-					static FSourceControlFileStatusMonitor::FOnSourceControlFileStatus EmptyDelegate{};
+						const FTypedElementPackageUnresolvedReference& UnresolvedReference = UnresolvedReferences[UnresolvedReferenceIndex];
+						if (UnresolvedReference.Index == 0)
+						{
+							Context.RemoveColumns<FTypedElementPackageUnresolvedReference>(RowHandles[UnresolvedReferenceIndex]);
+							return;
+						}
+						static FSourceControlFileStatusMonitor::FOnSourceControlFileStatus EmptyDelegate{};
 					
-					FileStatusMonitor.StartMonitoringFile(
-						reinterpret_cast<uintptr_t>(this),
-						UnresolvedReference.PathOnDisk,
-						EmptyDelegate
-					);
+						FileStatusMonitor.StartMonitoringFile(
+							reinterpret_cast<uintptr_t>(this),
+							UnresolvedReference.PathOnDisk,
+							EmptyDelegate
+						);
+					}
 				}
-			}
-		)
-		.Compile()
-	);
+			)
+			.Compile()
+		);
+	}
 }
 
 void UTypedElementRevisionControlFactory::RegisterApplyOverlays(ITypedElementDataStorageInterface& DataStorage) const
@@ -218,38 +233,69 @@ void UTypedElementRevisionControlFactory::RegisterApplyOverlays(ITypedElementDat
 	using namespace TypedElementQueryBuilder;
 	using DSI = ITypedElementDataStorageInterface;
 
-	ApplyOverlaysObjectToSCC = DataStorage.RegisterQuery(
-		Select()
-			.ReadOnly<FTypedElementPackagePathColumn>()
-			.ReadOnly<FSCCStatusColumn, FSCCExternallyLockedColumn>(EOptional::Yes)
-		.Compile());
-
-	ApplyOverlays = DataStorage.RegisterQuery(
-		Select(
-			TEXT("Change selection outline colors based on SCC status"),
-			// This is in PrePhysics because the outline->actor query is in DuringPhysics and contexts don't flush changes between tick groups
-			FProcessor(DSI::EQueryTickPhase::PrePhysics, DataStorage.GetQueryTickGroupName(DSI::EQueryTickGroups::SyncExternalToDataStorage))
-				.ForceToGameThread(true),
-			[](DSI::IQueryContext& Context, TypedElementRowHandle ObjectRow, const FTypedElementUObjectColumn& Actor, const FTypedElementPackageReference& PackageReference)
-			{
-				Context.RemoveColumns<FTypedElementViewportOverlayColorColumn>(ObjectRow);
-				Context.RunSubquery(0, PackageReference.Row, CreateSubqueryCallbackBinding(
-					[&Context, &ObjectRow, &Actor](DSI::ISubqueryContext& SubQueryContext)
-					{
-						FColor Color = DetermineOverlayColor(Context, SubQueryContext, Actor);
-						if (Color.Bits != 0)
+	if (ApplyOverlaysObjectToSCC == TypedElementInvalidQueryHandle)
+	{
+		ApplyOverlaysObjectToSCC = DataStorage.RegisterQuery(
+			Select()
+				.ReadOnly<FTypedElementPackagePathColumn>()
+				.ReadOnly<FSCCStatusColumn, FSCCExternallyLockedColumn>(EOptional::Yes)
+			.Compile());
+		}
+
+	if (ApplyOverlays == TypedElementInvalidQueryHandle)
+	{
+		ApplyOverlays = DataStorage.RegisterQuery(
+			Select(
+				TEXT("Change selection overlay colors based on SCC status"),
+				// This is in PrePhysics because the overlay->actor query is in DuringPhysics and contexts don't flush changes between tick groups
+				FProcessor(DSI::EQueryTickPhase::PrePhysics, DataStorage.GetQueryTickGroupName(DSI::EQueryTickGroups::SyncExternalToDataStorage))
+					.ForceToGameThread(true),
+				[](DSI::IQueryContext& Context, TypedElementRowHandle ObjectRow, const FTypedElementUObjectColumn& Actor, const FTypedElementPackageReference& PackageReference)
+				{
+					Context.RemoveColumns<FTypedElementViewportOverlayColorColumn>(ObjectRow);
+					Context.RunSubquery(0, PackageReference.Row, CreateSubqueryCallbackBinding(
+						[&Context, &ObjectRow, &Actor](DSI::ISubqueryContext& SubQueryContext)
 						{
-							Context.AddColumn<FTypedElementViewportOverlayColorColumn>(ObjectRow, { .OverlayColor = Color });
-						}
-					})
-				);
-				Context.AddColumns<FTypedElementSyncBackToWorldTag>(ObjectRow);
-			}
-		)
-		.Where()
-			.All<FTypedElementActorTag>()
-		.DependsOn()
-			.SubQuery(ApplyOverlaysObjectToSCC)
-		.Compile()
-	);
+							FColor Color = DetermineOverlayColor(Context, SubQueryContext, Actor);
+							if (Color.Bits != 0)
+							{
+								Context.AddColumn<FTypedElementViewportOverlayColorColumn>(ObjectRow, { .OverlayColor = Color });
+							}
+						})
+					);
+					Context.AddColumns<FTypedElementSyncBackToWorldTag>(ObjectRow);
+				}
+			)
+			.Where()
+				.All<FTypedElementActorTag>()
+			.DependsOn()
+				.SubQuery(ApplyOverlaysObjectToSCC)
+			.Compile()
+		);
+	}
+}
+
+void UTypedElementRevisionControlFactory::RegisterRemoveOverlays(ITypedElementDataStorageInterface& DataStorage) const
+{
+	using namespace TypedElementQueryBuilder;
+	using DSI = ITypedElementDataStorageInterface;
+
+	if (RemoveOverlays == TypedElementInvalidQueryHandle)
+	{
+		RemoveOverlays = DataStorage.RegisterQuery(
+			Select(
+				TEXT("Remove selection overlay colors"),
+				// This is in PrePhysics because the overlay->actor query is in DuringPhysics and contexts don't flush changes between tick groups
+				FProcessor(DSI::EQueryTickPhase::PrePhysics, DataStorage.GetQueryTickGroupName(DSI::EQueryTickGroups::SyncExternalToDataStorage))
+				.ForceToGameThread(true),
+				[](DSI::IQueryContext& Context, TypedElementRowHandle ObjectRow, FTypedElementUObjectColumn& Actor, const FTypedElementViewportOverlayColorColumn& ViewportColor)
+				{
+					Context.RemoveColumns<FTypedElementViewportOverlayColorColumn>(ObjectRow);
+				}
+			)
+			.Where()
+				.All<FTypedElementActorTag>()
+			.Compile()
+		);
+	}
 }
\ No newline at end of file
diff --git a/Engine/Plugins/Experimental/TypedElementsDataStorage/Source/TypedElementsDataStorageRevisionControl/Source/Private/RevisionControlProcessors.h b/Engine/Plugins/Experimental/TypedElementsDataStorage/Source/TypedElementsDataStorageRevisionControl/Source/Private/RevisionControlProcessors.h
--- a/Engine/Plugins/Experimental/TypedElementsDataStorage/Source/TypedElementsDataStorageRevisionControl/Source/Private/RevisionControlProcessors.h
+++ b/Engine/Plugins/Experimental/TypedElementsDataStorage/Source/TypedElementsDataStorageRevisionControl/Source/Private/RevisionControlProcessors.h
@@ -22,7 +22,9 @@ public:
 private:
 	void RegisterFetchUpdates(ITypedElementDataStorageInterface& DataStorage) const;
 	void RegisterApplyOverlays(ITypedElementDataStorageInterface& DataStorage) const;
+	void RegisterRemoveOverlays(ITypedElementDataStorageInterface& DataStorage) const;
 	mutable TypedElementDataStorage::QueryHandle FetchUpdates = TypedElementDataStorage::InvalidQueryHandle;
 	mutable TypedElementDataStorage::QueryHandle ApplyOverlays = TypedElementDataStorage::InvalidQueryHandle;
 	mutable TypedElementDataStorage::QueryHandle ApplyOverlaysObjectToSCC = TypedElementDataStorage::InvalidQueryHandle;
+	mutable TypedElementDataStorage::QueryHandle RemoveOverlays = TypedElementDataStorage::InvalidQueryHandle;
 };
