[TEDS] Check TEDS row handle is valid before passing off to Mass to check of valid entity

diff --git a/Engine/Plugins/Experimental/TypedElementsDataStorage/Source/TypedElementsDataStorage/Private/TypedElementDatabaseCommandBuffer.cpp b/Engine/Plugins/Experimental/TypedElementsDataStorage/Source/TypedElementsDataStorage/Private/TypedElementDatabaseCommandBuffer.cpp
--- a/Engine/Plugins/Experimental/TypedElementsDataStorage/Source/TypedElementsDataStorage/Private/TypedElementDatabaseCommandBuffer.cpp
+++ b/Engine/Plugins/Experimental/TypedElementsDataStorage/Source/TypedElementsDataStorage/Private/TypedElementDatabaseCommandBuffer.cpp
@@ -170,11 +170,19 @@ void FTypedElementDatabaseCommandBuffer::Queue_RemoveColumnsCommand(TypedElement
 
 bool FTypedElementDatabaseCommandBuffer::Execute_IsRowAvailable(const FMassEntityManager& MassEntityManager, TypedElementDataStorage::RowHandle Row)
 {
+	if (Row == TypedElementInvalidRowHandle)
+	{
+		return false;
+	}
 	return MassEntityManager.IsEntityValid(FMassEntityHandle::FromNumber(Row));
 }
 
 bool FTypedElementDatabaseCommandBuffer::Execute_HasRowBeenAssigned(const FMassEntityManager& MassEntityManager, TypedElementDataStorage::RowHandle Row)
 {
+	if (Row == TypedElementInvalidRowHandle)
+	{
+		return false;
+	}
 	return MassEntityManager.IsEntityActive(FMassEntityHandle::FromNumber(Row));
 }
 
