Fix continual unnecessary data sync from TEDS to world when revision control overlays are active
* Removed the the SyncBackToWorld tag addition for all actor rows
* Converted the processor that applies overlay color to UPrimitiveComponents to an observer OnAdd of OverlayColorColumn
* Replaced color determination processors with two queries:
1. Adds an OverlayColorColumn if one doesn't exist iff the overlay color determination is non-zero
2. Checks for overlay color changes by running color determination.  If changes occur, invoke the observer for the application of overlay color to UPrimitiveComponents on the Actor
* Added some debugging to validate the changes.  Only available in non-shipping builds.

diff --git a/Engine/Plugins/Experimental/TypedElementsDataStorage/Source/TypedElementsDataStorage/Private/Compatibility/TypedElementActorViewportProcessors.cpp b/Engine/Plugins/Experimental/TypedElementsDataStorage/Source/TypedElementsDataStorage/Private/Compatibility/TypedElementActorViewportProcessors.cpp
--- a/Engine/Plugins/Experimental/TypedElementsDataStorage/Source/TypedElementsDataStorage/Private/Compatibility/TypedElementActorViewportProcessors.cpp
+++ b/Engine/Plugins/Experimental/TypedElementsDataStorage/Source/TypedElementsDataStorage/Private/Compatibility/TypedElementActorViewportProcessors.cpp
@@ -206,7 +206,7 @@ void UTypedElementActorViewportFactory::RegisterOverlayColorColumnToActor(ITyped
 	DataStorage.RegisterQuery(
 		Select(
 			TEXT("Sync viewport overlay color column to actor"),
-			FProcessor(DSI::EQueryTickPhase::DuringPhysics, DataStorage.GetQueryTickGroupName(DSI::EQueryTickGroups::SyncDataStorageToExternal))
+			FObserver::OnAdd<FTypedElementViewportOverlayColorColumn>()
 				.ForceToGameThread(true),
 			[](FTypedElementUObjectColumn& Actor, const FTypedElementViewportOverlayColorColumn& ViewportColor)
 			{
@@ -221,7 +221,7 @@ void UTypedElementActorViewportFactory::RegisterOverlayColorColumnToActor(ITyped
 			}
 		)
 		.Where()
-			.All<FTypedElementActorTag, FTypedElementSyncBackToWorldTag>()
+			.All<FTypedElementActorTag>()
 		.Compile()
 	);
 
diff --git a/Engine/Plugins/Experimental/TypedElementsDataStorage/Source/TypedElementsDataStorageRevisionControl/Source/Private/RevisionControlProcessors.cpp b/Engine/Plugins/Experimental/TypedElementsDataStorage/Source/TypedElementsDataStorageRevisionControl/Source/Private/RevisionControlProcessors.cpp
--- a/Engine/Plugins/Experimental/TypedElementsDataStorage/Source/TypedElementsDataStorageRevisionControl/Source/Private/RevisionControlProcessors.cpp
+++ b/Engine/Plugins/Experimental/TypedElementsDataStorage/Source/TypedElementsDataStorageRevisionControl/Source/Private/RevisionControlProcessors.cpp
@@ -13,6 +13,7 @@
 #include "Elements/Columns/TypedElementRevisionControlColumns.h"
 #include "Elements/Columns/TypedElementSelectionColumns.h"
 #include "Elements/Columns/TypedElementViewportColumns.h"
+#include "Elements/Common/TypedElementDataStorageLog.h"
 #include "Elements/Framework/TypedElementQueryBuilder.h"
 #include "Elements/Framework/TypedElementRegistry.h"
 
@@ -60,10 +61,44 @@ TAutoConsoleVariable<int32> CVarOverlayAlpha(
 	TEXT("Configures overlay opacity."),
 	ECVF_Default);
 
+#if UE_BUILD_SHIPPING
+#define ENABLE_OVERLAY_DEBUG 0
+#else
+#define ENABLE_OVERLAY_DEBUG 1
+#endif
+
+#if ENABLE_OVERLAY_DEBUG
+static int32 gDefaultDebugForceColorOnAllValue = false; // [0..100]
+TAutoConsoleVariable<int32> CVarDebugForceColorOnAll(
+	TEXT("SourceControl.Overlays.Debug.ForceColorOnAll"),
+	gDefaultDebugForceColorOnAllValue,
+	TEXT("Debug to force overlay color on everything. 1 = Red, 2 = Green, 3 = Blue, 4 = White. 0 = off  ."),
+	ECVF_Default);
+#endif
+
 static FColor DetermineOverlayColor(const TypedElementDataStorage::IQueryContext& ObjectContext, const TypedElementDataStorage::ICommonQueryContext& SCCContext, const FTypedElementUObjectColumn& Actor)
 {
 	check(IsInGameThread());
 
+#if ENABLE_OVERLAY_DEBUG
+	if (int32 Force = CVarDebugForceColorOnAll.GetValueOnGameThread())
+	{
+		int32 Alpha = FMath::Lerp<float>(0.f, 255.f, gOverlayAlpha / 100.f);
+		switch(Force)
+		{
+		case 1:
+			return FColor(255, 0, 0, Alpha);
+		case 2:
+			return FColor(0, 255, 0, Alpha);
+		case 3:
+			return FColor(0, 0, 255, Alpha);
+		case 4:
+			return FColor(255, 255, 255, Alpha);
+		// Do normal determination for higher than 4
+		}
+	}
+#endif
+
 	bool bExternal = Actor.Object.IsValid() ? Cast<AActor>(Actor.Object)->IsPackageExternal() : false;
 	bool bIgnored = !bExternal;
 	bool bSelected = ObjectContext.HasColumn<FTypedElementSelectionColumn>();
@@ -160,8 +195,10 @@ void UTypedElementRevisionControlFactory::RegisterQueries(ITypedElementDataStora
 			}
 			else
 			{
-				DataStorage.UnregisterQuery(ApplyOverlays);
-				ApplyOverlays = TypedElementInvalidQueryHandle;
+				DataStorage.UnregisterQuery(ApplyNewOverlays);
+				ApplyNewOverlays = TypedElementInvalidQueryHandle;
+				DataStorage.UnregisterQuery(ChangeOverlay);
+				ChangeOverlay = TypedElementInvalidQueryHandle;
 
 				DataStorage.UnregisterQuery(ApplyOverlaysObjectToSCC);
 				ApplyOverlaysObjectToSCC = TypedElementInvalidQueryHandle;
@@ -238,21 +275,24 @@ void UTypedElementRevisionControlFactory::RegisterApplyOverlays(ITypedElementDat
 		ApplyOverlaysObjectToSCC = DataStorage.RegisterQuery(
 			Select()
 				.ReadOnly<FTypedElementPackagePathColumn>()
-				.ReadOnly<FSCCStatusColumn, FSCCExternallyLockedColumn>(EOptional::Yes)
+				.ReadOnly<FSCCStatusColumn, FSCCExternallyLockedColumn, FTypedElementViewportOverlayColorColumn>(EOptional::Yes)
 			.Compile());
 		}
 
-	if (ApplyOverlays == TypedElementInvalidQueryHandle)
+	if (ApplyNewOverlays == TypedElementInvalidQueryHandle)
 	{
-		ApplyOverlays = DataStorage.RegisterQuery(
+		// Query:
+		// For all actors without an overlay color column AND having a package reference:
+		//   Determine if a color should be applied based on SCC status tags
+		//   If so, add the OverlayColorColumn to the actor row
+		ApplyNewOverlays = DataStorage.RegisterQuery(
 			Select(
-				TEXT("Change selection overlay colors based on SCC status"),
+				TEXT("Poll to add selection overlay colors based on SCC status"),
 				// This is in PrePhysics because the overlay->actor query is in DuringPhysics and contexts don't flush changes between tick groups
 				FProcessor(DSI::EQueryTickPhase::PrePhysics, DataStorage.GetQueryTickGroupName(DSI::EQueryTickGroups::SyncExternalToDataStorage))
 					.ForceToGameThread(true),
 				[](DSI::IQueryContext& Context, TypedElementRowHandle ObjectRow, const FTypedElementUObjectColumn& Actor, const FTypedElementPackageReference& PackageReference)
 				{
-					Context.RemoveColumns<FTypedElementViewportOverlayColorColumn>(ObjectRow);
 					Context.RunSubquery(0, PackageReference.Row, CreateSubqueryCallbackBinding(
 						[&Context, &ObjectRow, &Actor](DSI::ISubqueryContext& SubQueryContext)
 						{
@@ -263,7 +303,49 @@ void UTypedElementRevisionControlFactory::RegisterApplyOverlays(ITypedElementDat
 							}
 						})
 					);
-					Context.AddColumns<FTypedElementSyncBackToWorldTag>(ObjectRow);
+				}
+			)
+			.Where()
+				.All<FTypedElementActorTag>()
+				.None<FTypedElementViewportOverlayColorColumn>()
+			.DependsOn()
+				.SubQuery(ApplyOverlaysObjectToSCC)
+			.Compile()
+		);
+	}
+
+	if (ChangeOverlay == TypedElementInvalidQueryHandle)
+	{
+		// Query:
+		// For all actors WITH an overlay color column AND having a package reference:
+		//   Re-check the color that should be applied based on SCC status tags
+		//   If the color has changed, remove and re-add the OverlayColorColumn to the actor row
+		//
+		// Note: Remove and re-add will trigger observer in TypedElementActorViewportProcessors to SetOverlayColor on the primitive components
+		ChangeOverlay = DataStorage.RegisterQuery(
+			Select(
+				TEXT("Poll for changed selection overlay colors based on SCC status"),
+				// This is in PrePhysics because the overlay->actor query is in DuringPhysics and contexts don't flush changes between tick groups
+				FProcessor(DSI::EQueryTickPhase::PrePhysics, DataStorage.GetQueryTickGroupName(DSI::EQueryTickGroups::SyncExternalToDataStorage))
+					.ForceToGameThread(true),
+				[](DSI::IQueryContext& Context, TypedElementRowHandle ObjectRow, const FTypedElementUObjectColumn& Actor, const FTypedElementPackageReference& PackageReference, const FTypedElementViewportOverlayColorColumn& OverlayColorColumn)
+				{
+					Context.RunSubquery(0, PackageReference.Row, CreateSubqueryCallbackBinding(
+						[&Context, &ObjectRow, &Actor, &OverlayColorColumn](DSI::ISubqueryContext& SubQueryContext)
+						{
+							FColor Color = DetermineOverlayColor(Context, SubQueryContext, Actor);
+							if (Color.Bits == 0)
+							{
+								Context.RemoveColumns<FTypedElementViewportOverlayColorColumn>(ObjectRow);
+							}
+							else if (Color != OverlayColorColumn.OverlayColor)
+							{
+								// Remove and re-add to trigger the observer
+								Context.RemoveColumns<FTypedElementViewportOverlayColorColumn>(ObjectRow);
+								Context.AddColumn<FTypedElementViewportOverlayColorColumn>(ObjectRow, { .OverlayColor = Color });
+							}
+						})
+					);
 				}
 			)
 			.Where()
@@ -282,6 +364,11 @@ void UTypedElementRevisionControlFactory::RegisterRemoveOverlays(ITypedElementDa
 
 	if (RemoveOverlays == TypedElementInvalidQueryHandle)
 	{
+		// Query:
+		// For all actors WITH an overlay color column AND having a package reference:
+		//   Remove the overlay color column
+		//
+		// This query is used to clean up the color columns if the overlay feature is disabled dynamically
 		RemoveOverlays = DataStorage.RegisterQuery(
 			Select(
 				TEXT("Remove selection overlay colors"),
diff --git a/Engine/Plugins/Experimental/TypedElementsDataStorage/Source/TypedElementsDataStorageRevisionControl/Source/Private/RevisionControlProcessors.h b/Engine/Plugins/Experimental/TypedElementsDataStorage/Source/TypedElementsDataStorageRevisionControl/Source/Private/RevisionControlProcessors.h
--- a/Engine/Plugins/Experimental/TypedElementsDataStorage/Source/TypedElementsDataStorageRevisionControl/Source/Private/RevisionControlProcessors.h
+++ b/Engine/Plugins/Experimental/TypedElementsDataStorage/Source/TypedElementsDataStorageRevisionControl/Source/Private/RevisionControlProcessors.h
@@ -24,7 +24,8 @@ private:
 	void RegisterApplyOverlays(ITypedElementDataStorageInterface& DataStorage) const;
 	void RegisterRemoveOverlays(ITypedElementDataStorageInterface& DataStorage) const;
 	mutable TypedElementDataStorage::QueryHandle FetchUpdates = TypedElementDataStorage::InvalidQueryHandle;
-	mutable TypedElementDataStorage::QueryHandle ApplyOverlays = TypedElementDataStorage::InvalidQueryHandle;
+	mutable TypedElementDataStorage::QueryHandle ApplyNewOverlays = TypedElementDataStorage::InvalidQueryHandle;
+	mutable TypedElementDataStorage::QueryHandle ChangeOverlay = TypedElementDataStorage::InvalidQueryHandle;
 	mutable TypedElementDataStorage::QueryHandle ApplyOverlaysObjectToSCC = TypedElementDataStorage::InvalidQueryHandle;
 	mutable TypedElementDataStorage::QueryHandle RemoveOverlays = TypedElementDataStorage::InvalidQueryHandle;
 };
