Convert RevisionControl TEDS processor to an observer to only register package with file monitor once
* Addresses a pathological performance issue that can occur with slow revision control systems or disabled revision control.  UObject columns have unresolved package columns and are persistently added to the file status monitor.  As the number of rows grow, more are added to the monitor.
* If the monitor doesn't convert UnresolvedPackageColumns to Resolved fast enough, then the query becomes quite large with the majority of the time spent re-adding the package to the file monitor unneccessarily.
* The quick fix for now is to only add the files when UnresolvedPackageColumns are added to the row.  Packages are not ever removed from the file monitor right now - as it existed prior to this change.
* NOTE: Tiny unrelated cleanup to remove the unnecessary mutables and convert to use a TypedElementDataStorage::InvalidQueryHandle

diff --git a/Engine/Plugins/Experimental/TypedElementsDataStorage/Source/TypedElementsDataStorageRevisionControl/Source/Private/RevisionControlProcessors.cpp b/Engine/Plugins/Experimental/TypedElementsDataStorage/Source/TypedElementsDataStorageRevisionControl/Source/Private/RevisionControlProcessors.cpp
--- a/Engine/Plugins/Experimental/TypedElementsDataStorage/Source/TypedElementsDataStorageRevisionControl/Source/Private/RevisionControlProcessors.cpp
+++ b/Engine/Plugins/Experimental/TypedElementsDataStorage/Source/TypedElementsDataStorageRevisionControl/Source/Private/RevisionControlProcessors.cpp
@@ -223,41 +223,30 @@ void UTypedElementRevisionControlFactory::RegisterQueries(ITypedElementDataStora
 	}
 }
 
-void UTypedElementRevisionControlFactory::RegisterFetchUpdates(ITypedElementDataStorageInterface& DataStorage) const
+void UTypedElementRevisionControlFactory::RegisterFetchUpdates(ITypedElementDataStorageInterface& DataStorage)
 {
 	using namespace TypedElementQueryBuilder;
 	using DSI = ITypedElementDataStorageInterface;
 	
 	FSourceControlFileStatusMonitor& FileStatusMonitor = ISourceControlModule::Get().GetSourceControlFileStatusMonitor();
 
-	if (FetchUpdates == TypedElementInvalidQueryHandle)
+	if (FetchUpdates == TypedElementDataStorage::InvalidQueryHandle)
 	{
 		FetchUpdates = DataStorage.RegisterQuery(
 			Select(
 				TEXT("Gather source control statuses for objects with unresolved package paths"),
-				FProcessor(DSI::EQueryTickPhase::DuringPhysics, DataStorage.GetQueryTickGroupName(DSI::EQueryTickGroups::SyncExternalToDataStorage))
-					.ForceToGameThread(true),
-				[this, &FileStatusMonitor](DSI::IQueryContext& Context, const FTypedElementPackageUnresolvedReference* InUnresolvedReferences)
+				FObserver::OnAdd<FTypedElementPackageUnresolvedReference>()
+				.ForceToGameThread(true),
+				[this, &FileStatusMonitor](DSI::IQueryContext& Context, const FTypedElementPackageUnresolvedReference& UnresolvedReference)
 				{
-					TConstArrayView<TypedElementDataStorage::RowHandle> RowHandles = Context.GetRowHandles();
-					TConstArrayView<FTypedElementPackageUnresolvedReference, int64> UnresolvedReferences { InUnresolvedReferences, Context.GetRowCount() };
-
-					for (int64 UnresolvedReferenceIndex = 0; UnresolvedReferenceIndex < UnresolvedReferences.Num(); ++UnresolvedReferenceIndex)
-					{
-						const FTypedElementPackageUnresolvedReference& UnresolvedReference = UnresolvedReferences[UnresolvedReferenceIndex];
-						if (UnresolvedReference.Index == 0)
-						{
-							Context.RemoveColumns<FTypedElementPackageUnresolvedReference>(RowHandles[UnresolvedReferenceIndex]);
-							return;
-						}
-						static FSourceControlFileStatusMonitor::FOnSourceControlFileStatus EmptyDelegate{};
-					
-						FileStatusMonitor.StartMonitoringFile(
-							reinterpret_cast<uintptr_t>(this),
-							UnresolvedReference.PathOnDisk,
-							EmptyDelegate
-						);
-					}
+					static FSourceControlFileStatusMonitor::FOnSourceControlFileStatus EmptyDelegate{};
+				
+					FileStatusMonitor.StartMonitoringFile(
+						reinterpret_cast<uintptr_t>(this),
+						UnresolvedReference.PathOnDisk,
+						EmptyDelegate
+						
+					);
 				}
 			)
 			.Compile()
@@ -265,7 +254,7 @@ void UTypedElementRevisionControlFactory::RegisterFetchUpdates(ITypedElementData
 	}
 }
 
-void UTypedElementRevisionControlFactory::RegisterApplyOverlays(ITypedElementDataStorageInterface& DataStorage) const
+void UTypedElementRevisionControlFactory::RegisterApplyOverlays(ITypedElementDataStorageInterface& DataStorage)
 {
 	using namespace TypedElementQueryBuilder;
 	using DSI = ITypedElementDataStorageInterface;
@@ -357,7 +346,7 @@ void UTypedElementRevisionControlFactory::RegisterApplyOverlays(ITypedElementDat
 	}
 }
 
-void UTypedElementRevisionControlFactory::RegisterRemoveOverlays(ITypedElementDataStorageInterface& DataStorage) const
+void UTypedElementRevisionControlFactory::RegisterRemoveOverlays(ITypedElementDataStorageInterface& DataStorage)
 {
 	using namespace TypedElementQueryBuilder;
 	using DSI = ITypedElementDataStorageInterface;
diff --git a/Engine/Plugins/Experimental/TypedElementsDataStorage/Source/TypedElementsDataStorageRevisionControl/Source/Private/RevisionControlProcessors.h b/Engine/Plugins/Experimental/TypedElementsDataStorage/Source/TypedElementsDataStorageRevisionControl/Source/Private/RevisionControlProcessors.h
--- a/Engine/Plugins/Experimental/TypedElementsDataStorage/Source/TypedElementsDataStorageRevisionControl/Source/Private/RevisionControlProcessors.h
+++ b/Engine/Plugins/Experimental/TypedElementsDataStorage/Source/TypedElementsDataStorageRevisionControl/Source/Private/RevisionControlProcessors.h
@@ -20,12 +20,12 @@ public:
 	void RegisterQueries(ITypedElementDataStorageInterface& DataStorage) override;
 
 private:
-	void RegisterFetchUpdates(ITypedElementDataStorageInterface& DataStorage) const;
-	void RegisterApplyOverlays(ITypedElementDataStorageInterface& DataStorage) const;
-	void RegisterRemoveOverlays(ITypedElementDataStorageInterface& DataStorage) const;
-	mutable TypedElementDataStorage::QueryHandle FetchUpdates = TypedElementDataStorage::InvalidQueryHandle;
-	mutable TypedElementDataStorage::QueryHandle ApplyNewOverlays = TypedElementDataStorage::InvalidQueryHandle;
-	mutable TypedElementDataStorage::QueryHandle ChangeOverlay = TypedElementDataStorage::InvalidQueryHandle;
-	mutable TypedElementDataStorage::QueryHandle ApplyOverlaysObjectToSCC = TypedElementDataStorage::InvalidQueryHandle;
-	mutable TypedElementDataStorage::QueryHandle RemoveOverlays = TypedElementDataStorage::InvalidQueryHandle;
+	void RegisterFetchUpdates(ITypedElementDataStorageInterface& DataStorage);
+	void RegisterApplyOverlays(ITypedElementDataStorageInterface& DataStorage);
+	void RegisterRemoveOverlays(ITypedElementDataStorageInterface& DataStorage);
+	TypedElementDataStorage::QueryHandle FetchUpdates = TypedElementDataStorage::InvalidQueryHandle;
+	TypedElementDataStorage::QueryHandle ApplyNewOverlays = TypedElementDataStorage::InvalidQueryHandle;
+	TypedElementDataStorage::QueryHandle ChangeOverlay = TypedElementDataStorage::InvalidQueryHandle;
+	TypedElementDataStorage::QueryHandle ApplyOverlaysObjectToSCC = TypedElementDataStorage::InvalidQueryHandle;
+	TypedElementDataStorage::QueryHandle RemoveOverlays = TypedElementDataStorage::InvalidQueryHandle;
 };
