Implemented AddRow for TEDS IQueryContext and ISubqueryContext
* Dependent on Mass Concurrent Entity Reserve changes
* Functional changes here allow TEDS processor queries and processor subqueries to add and remove rows.   Implementation for Add rows is that row will be added as reserved - manipulations such as add/remove column are already deferred.
* Refactored the Mass processor adapters to explicitly inherit from interfaces instead of through template - improves searchability to implementation from IDE and debugging callstack.
* Also added RemoveRow to SubqueryContext as there wasn't a good reason not to.

diff --git a/Engine/Plugins/Experimental/TypedElementsDataStorage/Source/TypedElementsDataStorage/Private/Compatibility/TypedElementObjectReinstancingManager.cpp b/Engine/Plugins/Experimental/TypedElementsDataStorage/Source/TypedElementsDataStorage/Private/Compatibility/TypedElementObjectReinstancingManager.cpp
--- a/Engine/Plugins/Experimental/TypedElementsDataStorage/Source/TypedElementsDataStorage/Private/Compatibility/TypedElementObjectReinstancingManager.cpp
+++ b/Engine/Plugins/Experimental/TypedElementsDataStorage/Source/TypedElementsDataStorage/Private/Compatibility/TypedElementObjectReinstancingManager.cpp
@@ -7,6 +7,7 @@
 #include "Memento/TypedElementMementoRowTypes.h"
 #include "TypedElementDatabase.h"
 #include "TypedElementDatabaseCompatibility.h"
+#include "TypedElementDatabaseEnvironment.h"
 
 DECLARE_LOG_CATEGORY_CLASS(LogTedsObjectReinstancing, Log, Log)
 
diff --git a/Engine/Plugins/Experimental/TypedElementsDataStorage/Source/TypedElementsDataStorage/Private/Processors/TypedElementProcessorAdaptors.cpp b/Engine/Plugins/Experimental/TypedElementsDataStorage/Source/TypedElementsDataStorage/Private/Processors/TypedElementProcessorAdaptors.cpp
--- a/Engine/Plugins/Experimental/TypedElementsDataStorage/Source/TypedElementsDataStorage/Private/Processors/TypedElementProcessorAdaptors.cpp
+++ b/Engine/Plugins/Experimental/TypedElementsDataStorage/Source/TypedElementsDataStorage/Private/Processors/TypedElementProcessorAdaptors.cpp
@@ -11,17 +11,20 @@
 #include "TypedElementDatabase.h"
 #include "TypedElementDatabaseEnvironment.h"
 
-template<typename T>
-struct FMassContextCommon : public T
+struct FMassContextCommon
 {
-	~FMassContextCommon() override = default;
-
-	uint32 GetRowCount() const override
+	FMassExecutionContext& Context;
+	
+	explicit FMassContextCommon(FMassExecutionContext& InContext)
+		: Context(InContext)
+	{}
+	
+	uint32 GetRowCount() const
 	{
 		return Context.GetNumEntities();
 	}
 
-	TConstArrayView<TypedElementRowHandle> GetRowHandles() const override
+	TConstArrayView<TypedElementRowHandle> GetRowHandles() const
 	{
 		static_assert(
 			sizeof(TypedElementRowHandle) == sizeof(FMassEntityHandle) && alignof(TypedElementRowHandle) == alignof(FMassEntityHandle),
@@ -30,18 +33,19 @@ struct FMassContextCommon : public T
 		return TConstArrayView<TypedElementRowHandle>(reinterpret_cast<const TypedElementRowHandle*>(Entities.GetData()), Entities.Num());
 	}
 	
-	const void* GetColumn(const UScriptStruct* ColumnType) const override
+	const void* GetColumn(const UScriptStruct* ColumnType) const
 	{
 		return Context.GetFragmentView(ColumnType).GetData();
 	}
 
-	void* GetMutableColumn(const UScriptStruct* ColumnType) override
+	void* GetMutableColumn(const UScriptStruct* ColumnType)
 	{
 		return Context.GetMutableFragmentView(ColumnType).GetData();
 	}
 
-	void GetColumns(TArrayView<char*> RetrievedAddresses, TConstArrayView<TWeakObjectPtr<const UScriptStruct>> ColumnTypes,
-		TConstArrayView<TypedElementDataStorage::EQueryAccessType> AccessTypes) override
+	void GetColumns(TArrayView<char*> RetrievedAddresses,
+		TConstArrayView<TWeakObjectPtr<const UScriptStruct>> ColumnTypes,
+		TConstArrayView<TypedElementDataStorage::EQueryAccessType> AccessTypes)
 	{
 		checkf(RetrievedAddresses.Num() == ColumnTypes.Num(), TEXT("Unable to retrieve a batch of columns as the number of addresses "
 			"doesn't match the number of requested column."));
@@ -51,8 +55,11 @@ struct FMassContextCommon : public T
 		GetColumnsUnguarded(ColumnTypes.Num(), RetrievedAddresses.GetData(), ColumnTypes.GetData(), AccessTypes.GetData());
 	}
 	
-	void GetColumnsUnguarded(int32 TypeCount, char** RetrievedAddresses, const TWeakObjectPtr<const UScriptStruct>* ColumnTypes,
-		const TypedElementDataStorage::EQueryAccessType* AccessTypes) override
+	void GetColumnsUnguarded(
+		int32 TypeCount,
+		char** RetrievedAddresses,
+		const TWeakObjectPtr<const UScriptStruct>* ColumnTypes,
+		const TypedElementDataStorage::EQueryAccessType* AccessTypes)
 	{
 		for (int32 Index = 0; Index < TypeCount; ++Index)
 		{
@@ -67,7 +74,7 @@ struct FMassContextCommon : public T
 		}
 	}
 
-	bool HasColumn(const UScriptStruct* ColumnType) const override
+	bool HasColumn(const UScriptStruct* ColumnType) const
 	{
 		if (ColumnType->IsChildOf(FMassTag::StaticStruct()))
 		{
@@ -82,7 +89,7 @@ struct FMassContextCommon : public T
 		return false;
 	}
 
-	bool HasColumn(TypedElementDataStorage::RowHandle Row, const UScriptStruct* ColumnType) const override
+	bool HasColumn(TypedElementDataStorage::RowHandle Row, const UScriptStruct* ColumnType) const
 	{
 		FMassEntityHandle Entity = FMassEntityHandle::FromNumber(Row);
 		FMassEntityManager& Manager = Context.GetEntityManagerChecked();
@@ -101,19 +108,11 @@ struct FMassContextCommon : public T
 		checkf(bIsTagOrFragment, TEXT("Attempting to check for a column type that is not a column or tag."));
 		return false;
 	}
-
-protected:
-	explicit FMassContextCommon(FMassExecutionContext& InContext)
-		: Context(InContext)
-	{}
-
-	FMassExecutionContext& Context;
 };
 
-template<typename T>
-struct FMassWithEnvironmentContextCommon : public FMassContextCommon<T>
+struct FMassWithEnvironmentContextCommon : public FMassContextCommon
 {
-	using Parent = FMassContextCommon<T>;
+	using Parent = FMassContextCommon;
 protected:
 	void TedsColumnsToMassDescriptorIfActiveTable(
 		FMassArchetypeCompositionDescriptor& Descriptor,
@@ -161,18 +160,28 @@ protected:
 	}
 
 public:
+	using ObjectCopyOrMove = void (*)(const UScriptStruct& TypeInfo, void* Destination, void* Source);
+	
 	FMassWithEnvironmentContextCommon(FMassExecutionContext& InContext, FTypedElementDatabaseEnvironment& InEnvironment)
-		: FMassContextCommon<T>(InContext)
+		: FMassContextCommon(InContext)
 		, Environment(InEnvironment)
 	{}
 
-	~FMassWithEnvironmentContextCommon() override = default;
-
-	uint64 GetUpdateCycleId() const override
+	uint64 GetUpdateCycleId() const 
 	{
 		return Environment.GetUpdateCycleId();
 	}
 
+	bool IsRowAvailable(TypedElementDataStorage::RowHandle Row) const
+	{
+		return Environment.GetMassEntityManager().IsEntityValid(FMassEntityHandle::FromNumber(Row));
+	}
+
+	bool HasRowBeenAssigned(TypedElementDataStorage::RowHandle Row) const
+	{
+		return Environment.GetMassEntityManager().IsEntityActive(FMassEntityHandle::FromNumber(Row));
+	}
+	
 	void ActivateQueries(FName ActivationName)
 	{
 		this->Context.Defer().template PushCommand<FMassDeferredCommand<EMassCommandOperationType::None>>(
@@ -182,7 +191,7 @@ public:
 			});
 	}
 
-	void* AddColumnUninitialized(TypedElementDataStorage::RowHandle Row, const UScriptStruct* ObjectType) override
+	void* AddColumnUninitialized(TypedElementDataStorage::RowHandle Row, const UScriptStruct* ObjectType) 
 	{
 		return AddColumnUninitialized(Row, ObjectType,
 			[](const UScriptStruct& TypeInfo, void* Destination, void* Source)
@@ -191,19 +200,19 @@ public:
 			});
 	}
 
-	void* AddColumnUninitialized(TypedElementDataStorage::RowHandle Row, const UScriptStruct* ObjectType, typename Parent::ObjectCopyOrMove Relocator) override
+	void* AddColumnUninitialized(TypedElementDataStorage::RowHandle Row, const UScriptStruct* ObjectType, ObjectCopyOrMove Relocator)
 	{
 		checkf(ObjectType->IsChildOf(FMassFragment::StaticStruct()), TEXT("Column [%s] can not be a tag"), *ObjectType->GetName());
 		
 		struct FAddValueColumn
 		{
-			typename Parent::ObjectCopyOrMove Relocator;
+			ObjectCopyOrMove Relocator;
 			const UScriptStruct* FragmentType;
 			FMassEntityHandle Entity;
 			void* Object;
 
 			FAddValueColumn() = default;
-			FAddValueColumn(typename Parent::ObjectCopyOrMove InRelocator, const UScriptStruct* InFragmentType, FMassEntityHandle InEntity, void* InObject)
+			FAddValueColumn(ObjectCopyOrMove InRelocator, const UScriptStruct* InFragmentType, FMassEntityHandle InEntity, void* InObject)
 				: Relocator(InRelocator)
 				, FragmentType(InFragmentType)
 				, Entity(InEntity)
@@ -251,7 +260,7 @@ public:
 		return ColumnData;
 	}
 
-	void AddColumns(TypedElementDataStorage::RowHandle Row, TConstArrayView<const UScriptStruct*> ColumnTypes) override
+	void AddColumns(TypedElementDataStorage::RowHandle Row, TConstArrayView<const UScriptStruct*> ColumnTypes) 
 	{
 		struct FAddedColumns
 		{
@@ -273,7 +282,7 @@ public:
 			});
 	}
 
-	void AddColumns(TConstArrayView<TypedElementDataStorage::RowHandle> Rows, TConstArrayView<const UScriptStruct*> ColumnTypes) override
+	void AddColumns(TConstArrayView<TypedElementDataStorage::RowHandle> Rows, TConstArrayView<const UScriptStruct*> ColumnTypes) 
 	{
 		struct FAddedColumns
 		{
@@ -310,7 +319,7 @@ public:
 			});
 	}
 
-	void RemoveColumns(TypedElementDataStorage::RowHandle Row, TConstArrayView<const UScriptStruct*> ColumnTypes) override
+	void RemoveColumns(TypedElementDataStorage::RowHandle Row, TConstArrayView<const UScriptStruct*> ColumnTypes) 
 	{
 		struct FRemovedColumns
 		{
@@ -335,7 +344,7 @@ public:
 		}
 	}
 
-	void RemoveColumns(TConstArrayView<TypedElementDataStorage::RowHandle> Rows, TConstArrayView<const UScriptStruct*> ColumnTypes) override
+	void RemoveColumns(TConstArrayView<TypedElementDataStorage::RowHandle> Rows, TConstArrayView<const UScriptStruct*> ColumnTypes) 
 	{
 		struct FRemovedColumns
 		{
@@ -371,53 +380,140 @@ public:
 				}
 			});
 	}
+	
+	TypedElementDataStorage::RowHandle AddRow(TypedElementDataStorage::TableHandle Table) 
+	{
+		FMassEntityHandle EntityHandle = Environment.GetMassEntityManager().ReserveEntity();
+		FMassArchetypeHandle ArchetypeHandle = Environment.LookupMassArchetype(Table);
+
+		if (!ArchetypeHandle.IsValid())
+		{
+			return TypedElementDataStorage::InvalidRowHandle;
+		}
+		
+		struct CommandInfo
+		{
+			FMassEntityHandle Entity;
+			FMassArchetypeHandle Archetype;
+		};
+
+		CommandInfo CommandDataTmp{
+			.Entity = EntityHandle,
+			.Archetype = MoveTemp(ArchetypeHandle)
+		};
+
+		this->Context.Defer().template PushCommand<FMassDeferredCreateCommand>(
+			[CommandData = MoveTemp(CommandDataTmp)](FMassEntityManager& System)
+			{
+				const FMassArchetypeSharedFragmentValues SharedFragmentValues;
+				System.BuildEntity(CommandData.Entity, CommandData.Archetype, SharedFragmentValues);
+			});
+		
+		const TypedElementDataStorage::RowHandle TedsRowHandle = EntityHandle.AsNumber();
+		return TedsRowHandle;
+	}
+
+	void RemoveRow(TypedElementRowHandle Row) 
+	{
+		this->Context.Defer().DestroyEntity(FMassEntityHandle::FromNumber(Row));
+	}
+
+	void RemoveRows(TConstArrayView<TypedElementRowHandle> Rows) 
+	{
+		// Row handles and entities map 1:1 for data, so a reintpret_cast can be safely done to avoid
+		// having to allocate memory and iterating over the rows.
+
+		static_assert(sizeof(FMassEntityHandle) == sizeof(TypedElementRowHandle), 
+			"Sizes of mass entity and data storage row have gone out of sync.");
+		static_assert(alignof(FMassEntityHandle) == alignof(TypedElementRowHandle),
+			"Alignment of mass entity and data storage row have gone out of sync.");
+
+		this->Context.Defer().DestroyEntities(
+			TConstArrayView<FMassEntityHandle>(reinterpret_cast<const FMassEntityHandle*>(Rows.begin()), Rows.Num()));
+	}
 
 protected:
 	FTypedElementDatabaseEnvironment& Environment;
 };
 
-struct FMassDirectContextForwarder final : public FMassContextCommon<ITypedElementDataStorageInterface::IDirectQueryContext>
+struct FMassDirectContextForwarder final : public ITypedElementDataStorageInterface::IDirectQueryContext
 {
 	explicit FMassDirectContextForwarder(FMassExecutionContext& InContext)
-		: FMassContextCommon(InContext)
+		: Implementation(InContext)
 	{}
-
-	~FMassDirectContextForwarder() override = default;
+	
+	uint32 GetRowCount() const override { return Implementation.GetRowCount(); }
+	TConstArrayView<TypedElementDataStorage::RowHandle> GetRowHandles() const override { return Implementation.GetRowHandles(); }
+	const void* GetColumn(const UScriptStruct* ColumnType) const override { return Implementation.GetColumn(ColumnType); }
+	void* GetMutableColumn(const UScriptStruct* ColumnType) override { return Implementation.GetMutableColumn(ColumnType); }
+	void GetColumns(TArrayView<char*> RetrievedAddresses, TConstArrayView<TWeakObjectPtr<const UScriptStruct>> ColumnTypes, TConstArrayView<TypedElementDataStorage::EQueryAccessType> AccessTypes) override { return Implementation.GetColumns(RetrievedAddresses, ColumnTypes, AccessTypes); }
+	void GetColumnsUnguarded(int32 TypeCount, char** RetrievedAddresses, const TWeakObjectPtr<const UScriptStruct>* ColumnTypes, const TypedElementDataStorage::EQueryAccessType* AccessTypes) override { return Implementation.GetColumnsUnguarded(TypeCount, RetrievedAddresses, ColumnTypes, AccessTypes); }
+	bool HasColumn(const UScriptStruct* ColumnType) const override { return Implementation.HasColumn(ColumnType); }
+	bool HasColumn(TypedElementDataStorage::RowHandle Row, const UScriptStruct* ColumnType) const override { return Implementation.HasColumn(Row, ColumnType); }
+
+	FMassContextCommon Implementation;
 };
 
-struct FMassSubqueryContextForwarder final : public FMassWithEnvironmentContextCommon<ITypedElementDataStorageInterface::ISubqueryContext>
+struct FMassSubqueryContextForwarder  : public ITypedElementDataStorageInterface::ISubqueryContext
 {
 	FMassSubqueryContextForwarder(FMassExecutionContext& InContext, FTypedElementDatabaseEnvironment& InEnvironment)
-		: FMassWithEnvironmentContextCommon(InContext, InEnvironment)
+		: Implementation(InContext, InEnvironment)
 	{}
 
 	~FMassSubqueryContextForwarder() override = default;
+	uint32 GetRowCount() const override { return Implementation.GetRowCount(); }
+	TConstArrayView<TypedElementDataStorage::RowHandle> GetRowHandles() const override { return Implementation.GetRowHandles(); }
+	const void* GetColumn(const UScriptStruct* ColumnType) const override { return Implementation.GetColumn(ColumnType); }
+	void* GetMutableColumn(const UScriptStruct* ColumnType) override { return Implementation.GetMutableColumn(ColumnType); }
+	void GetColumns(TArrayView<char*> RetrievedAddresses, TConstArrayView<TWeakObjectPtr<const UScriptStruct>> ColumnTypes, TConstArrayView<TypedElementDataStorage::EQueryAccessType> AccessTypes) override { return Implementation.GetColumns(RetrievedAddresses, ColumnTypes, AccessTypes); }
+	void GetColumnsUnguarded(int32 TypeCount, char** RetrievedAddresses, const TWeakObjectPtr<const UScriptStruct>* ColumnTypes, const TypedElementDataStorage::EQueryAccessType* AccessTypes) override { return Implementation.GetColumnsUnguarded(TypeCount, RetrievedAddresses, ColumnTypes, AccessTypes); }
+	bool HasColumn(const UScriptStruct* ColumnType) const override { return Implementation.HasColumn(ColumnType); }
+	bool HasColumn(TypedElementDataStorage::RowHandle Row, const UScriptStruct* ColumnType) const override { return Implementation.HasColumn(Row, ColumnType); }
+	uint64 GetUpdateCycleId() const override { return Implementation.GetUpdateCycleId(); }
+	bool IsRowAvailable(TypedElementDataStorage::RowHandle Row) const override { return Implementation.IsRowAvailable(Row); }
+	bool HasRowBeenAssigned(TypedElementDataStorage::RowHandle Row) const override { return Implementation.HasRowBeenAssigned(Row); }
+	void ActivateQueries(FName ActivationName) override  { return Implementation.ActivateQueries(ActivationName); }
+	TypedElementDataStorage::RowHandle AddRow(TypedElementDataStorage::TableHandle Table) override  { return Implementation.AddRow(Table); }
+	void RemoveRow(TypedElementDataStorage::RowHandle Row) override { return Implementation.RemoveRow(Row); }
+	void RemoveRows(TConstArrayView<TypedElementDataStorage::RowHandle> Rows) override  { return Implementation.RemoveRows(Rows); }
+	void AddColumns(TypedElementDataStorage::RowHandle Row, TConstArrayView<const UScriptStruct*> ColumnTypes) override { return Implementation.AddColumns(Row, ColumnTypes); }
+	void AddColumns(TConstArrayView<TypedElementDataStorage::RowHandle> Rows, TConstArrayView<const UScriptStruct*> ColumnTypes) override  { return Implementation.AddColumns(Rows, ColumnTypes); }
+	void* AddColumnUninitialized(TypedElementDataStorage::RowHandle Row, const UScriptStruct* ColumnType) override  { return Implementation.AddColumnUninitialized(Row, ColumnType); }
+	void* AddColumnUninitialized(TypedElementDataStorage::RowHandle Row, const UScriptStruct* ObjectType, ObjectCopyOrMove Relocator) override { return Implementation.AddColumnUninitialized(Row, ObjectType, Relocator); }
+	void RemoveColumns(TypedElementDataStorage::RowHandle Row, TConstArrayView<const UScriptStruct*> ColumnTypes) override { Implementation.RemoveColumns(Row, ColumnTypes); }
+	void RemoveColumns(TConstArrayView<TypedElementDataStorage::RowHandle> Rows, TConstArrayView<const UScriptStruct*> ColumnTypes) override { Implementation.RemoveColumns(Rows, ColumnTypes); }
+
+	FMassWithEnvironmentContextCommon Implementation;
 };
 
-struct FMassContextForwarder final : public FMassWithEnvironmentContextCommon<ITypedElementDataStorageInterface::IQueryContext>
+struct FMassQueryContextImplementation final : FMassWithEnvironmentContextCommon
 {
-public:
-	FMassContextForwarder(ITypedElementDataStorageInterface::FQueryDescription& InQueryDescription, FMassExecutionContext& InContext, 
-		FTypedElementExtendedQueryStore& InQueryStore, FTypedElementDatabaseEnvironment& InEnvironment)
+	FMassQueryContextImplementation(
+		ITypedElementDataStorageInterface::FQueryDescription& InQueryDescription,
+		FMassExecutionContext& InContext, 
+		FTypedElementExtendedQueryStore& InQueryStore,
+		FTypedElementDatabaseEnvironment& InEnvironment)
 		: FMassWithEnvironmentContextCommon(InContext, InEnvironment)
 		, QueryDescription(InQueryDescription)
 		, QueryStore(InQueryStore)
-	{}
+	{
+		
+	}
 
-	~FMassContextForwarder() override = default;
+	~FMassQueryContextImplementation() = default;
 
-	UObject* GetMutableDependency(const UClass* DependencyClass) override
+	UObject* GetMutableDependency(const UClass* DependencyClass)
 	{
 		return Context.GetMutableSubsystem<USubsystem>(const_cast<UClass*>(DependencyClass));
 	}
 
-	const UObject* GetDependency(const UClass* DependencyClass) override
+	const UObject* GetDependency(const UClass* DependencyClass)
 	{
 		return Context.GetSubsystem<USubsystem>(const_cast<UClass*>(DependencyClass));
 	}
-
+	
 	void GetDependencies(TArrayView<UObject*> RetrievedAddresses, TConstArrayView<TWeakObjectPtr<const UClass>> SubsystemTypes,
-		TConstArrayView<ITypedElementDataStorageInterface::EQueryAccessType> AccessTypes) override
+		TConstArrayView<ITypedElementDataStorageInterface::EQueryAccessType> AccessTypes)
 	{
 		checkf(RetrievedAddresses.Num() == SubsystemTypes.Num(), TEXT("Unable to retrieve a batch of subsystem as the number of addresses "
 			"doesn't match the number of requested subsystem types."));
@@ -440,56 +536,27 @@ public:
 			++AccessTypes;
 		}
 	}
-
-	bool IsRowAvailable(TypedElementDataStorage::RowHandle Row) const override
-	{
-		return Environment.GetMassEntityManager().IsEntityValid(FMassEntityHandle::FromNumber(Row));
-	}
-
-	bool HasRowBeenAssigned(TypedElementDataStorage::RowHandle Row) const override
-	{
-		return Environment.GetMassEntityManager().IsEntityActive(FMassEntityHandle::FromNumber(Row));
-	}
-
-	void RemoveRow(TypedElementRowHandle Row) override
-	{
-		Context.Defer().DestroyEntity(FMassEntityHandle::FromNumber(Row));
-	}
-
-	void RemoveRows(TConstArrayView<TypedElementRowHandle> Rows) override
-	{
-		// Row handles and entities map 1:1 for data, so a reintpret_cast can be safely done to avoid
-		// having to allocate memory and iterating over the rows.
-
-		static_assert(sizeof(FMassEntityHandle) == sizeof(TypedElementRowHandle), 
-			"Sizes of mass entity and data storage row have gone out of sync.");
-		static_assert(alignof(FMassEntityHandle) == alignof(TypedElementRowHandle),
-			"Alignment of mass entity and data storage row have gone out of sync.");
-
-		Context.Defer().DestroyEntities(
-			TConstArrayView<FMassEntityHandle>(reinterpret_cast<const FMassEntityHandle*>(Rows.begin()), Rows.Num()));
-	}
-
-	TypedElementDataStorage::RowHandle FindIndexedRow(TypedElementDataStorage::IndexHash Index) const override
+	
+	TypedElementDataStorage::RowHandle FindIndexedRow(TypedElementDataStorage::IndexHash Index) const
 	{
 		return Environment.GetIndexTable().FindIndexedRow(Index);
 	}
 
-	TypedElementDataStorage::FQueryResult RunQuery(TypedElementQueryHandle Query) override
+	TypedElementDataStorage::FQueryResult RunQuery(TypedElementQueryHandle Query)
 	{
 		const FTypedElementExtendedQueryStore::Handle Handle(Query);
 		// This can be safely called because there's not callback, which means no columns are accessed, even for select queries.
 		return QueryStore.RunQuery(Context.GetEntityManagerChecked(), Handle);
 	}
 
-	TypedElementDataStorage::FQueryResult RunSubquery(int32 SubqueryIndex) override
+	TypedElementDataStorage::FQueryResult RunSubquery(int32 SubqueryIndex)
 	{
 		return SubqueryIndex < QueryDescription.Subqueries.Num() ?
 			RunQuery(QueryDescription.Subqueries[SubqueryIndex]) :
 			TypedElementDataStorage::FQueryResult{};
 	}
 
-	TypedElementDataStorage::FQueryResult RunSubquery(int32 SubqueryIndex, TypedElementDataStorage::SubqueryCallbackRef Callback) override
+	TypedElementDataStorage::FQueryResult RunSubquery(int32 SubqueryIndex, TypedElementDataStorage::SubqueryCallbackRef Callback)
 	{
 		if (SubqueryIndex < QueryDescription.Subqueries.Num())
 		{
@@ -504,7 +571,7 @@ public:
 	}
 
 	TypedElementDataStorage::FQueryResult RunSubquery(int32 SubqueryIndex, TypedElementDataStorage::RowHandle Row,
-		TypedElementDataStorage::SubqueryCallbackRef Callback) override
+		TypedElementDataStorage::SubqueryCallbackRef Callback)
 	{
 		if (SubqueryIndex < QueryDescription.Subqueries.Num())
 		{
@@ -517,14 +584,54 @@ public:
 			return TypedElementDataStorage::FQueryResult{};
 		}
 	}
-
+	
 	ITypedElementDataStorageInterface::FQueryDescription& QueryDescription;
 	FTypedElementExtendedQueryStore& QueryStore;
 };
 
-
-
-
+struct FMassContextForwarder final : public ITypedElementDataStorageInterface::IQueryContext
+{
+	FMassContextForwarder(
+		ITypedElementDataStorageInterface::FQueryDescription& InQueryDescription,
+		FMassExecutionContext& InContext, 
+		FTypedElementExtendedQueryStore& InQueryStore,
+		FTypedElementDatabaseEnvironment& InEnvironment)
+			: Implementation(InQueryDescription, InContext, InQueryStore, InEnvironment)
+	{}
+	
+	uint32 GetRowCount() const override { return Implementation.GetRowCount(); }
+	TConstArrayView<TypedElementDataStorage::RowHandle> GetRowHandles() const override { return Implementation.GetRowHandles(); }
+	const void* GetColumn(const UScriptStruct* ColumnType) const override { return Implementation.GetColumn(ColumnType); }
+	void* GetMutableColumn(const UScriptStruct* ColumnType) override { return Implementation.GetMutableColumn(ColumnType); }
+	void GetColumns(TArrayView<char*> RetrievedAddresses, TConstArrayView<TWeakObjectPtr<const UScriptStruct>> ColumnTypes, TConstArrayView<TypedElementDataStorage::EQueryAccessType> AccessTypes) override { return Implementation.GetColumns(RetrievedAddresses, ColumnTypes, AccessTypes); }
+	void GetColumnsUnguarded(int32 TypeCount, char** RetrievedAddresses, const TWeakObjectPtr<const UScriptStruct>* ColumnTypes, const TypedElementDataStorage::EQueryAccessType* AccessTypes) override { return Implementation.GetColumnsUnguarded(TypeCount, RetrievedAddresses, ColumnTypes, AccessTypes); }
+	bool HasColumn(const UScriptStruct* ColumnType) const override { return Implementation.HasColumn(ColumnType); }
+	bool HasColumn(TypedElementDataStorage::RowHandle Row, const UScriptStruct* ColumnType) const override { return Implementation.HasColumn(Row, ColumnType); }
+	uint64 GetUpdateCycleId() const override { return Implementation.GetUpdateCycleId(); }
+	bool IsRowAvailable(TypedElementDataStorage::RowHandle Row) const override { return Implementation.IsRowAvailable(Row); }
+	bool HasRowBeenAssigned(TypedElementDataStorage::RowHandle Row) const override { return Implementation.HasRowBeenAssigned(Row); }
+	void ActivateQueries(FName ActivationName) override  { return Implementation.ActivateQueries(ActivationName); }
+	TypedElementDataStorage::RowHandle AddRow(TypedElementDataStorage::TableHandle Table) override  { return Implementation.AddRow(Table); }
+	void RemoveRow(TypedElementDataStorage::RowHandle Row) override { return Implementation.RemoveRow(Row); }
+	void RemoveRows(TConstArrayView<TypedElementDataStorage::RowHandle> Rows) override  { return Implementation.RemoveRows(Rows); }
+	void AddColumns(TypedElementDataStorage::RowHandle Row, TConstArrayView<const UScriptStruct*> ColumnTypes) override { return Implementation.AddColumns(Row, ColumnTypes); }
+	void AddColumns(TConstArrayView<TypedElementDataStorage::RowHandle> Rows, TConstArrayView<const UScriptStruct*> ColumnTypes) override  { return Implementation.AddColumns(Rows, ColumnTypes); }
+	void* AddColumnUninitialized(TypedElementDataStorage::RowHandle Row, const UScriptStruct* ColumnType) override  { return Implementation.AddColumnUninitialized(Row, ColumnType); }
+	void* AddColumnUninitialized(TypedElementDataStorage::RowHandle Row, const UScriptStruct* ObjectType, ObjectCopyOrMove Relocator) override { return Implementation.AddColumnUninitialized(Row, ObjectType, Relocator); }
+	void RemoveColumns(TypedElementDataStorage::RowHandle Row, TConstArrayView<const UScriptStruct*> ColumnTypes) override { Implementation.RemoveColumns(Row, ColumnTypes); }
+	void RemoveColumns(TConstArrayView<TypedElementDataStorage::RowHandle> Rows, TConstArrayView<const UScriptStruct*> ColumnTypes) override { Implementation.RemoveColumns(Rows, ColumnTypes); }
+	
+	const UObject* GetDependency(const UClass* DependencyClass) override { return Implementation.GetDependency(DependencyClass); }
+	UObject* GetMutableDependency(const UClass* DependencyClass) override { return Implementation.GetMutableDependency(DependencyClass); }
+	void GetDependencies(TArrayView<UObject*> RetrievedAddresses, TConstArrayView<TWeakObjectPtr<const UClass>> DependencyTypes, TConstArrayView<TypedElementDataStorage::EQueryAccessType> AccessTypes) override { return Implementation.GetDependencies(RetrievedAddresses, DependencyTypes, AccessTypes); }
+	TypedElementDataStorage::RowHandle FindIndexedRow(TypedElementDataStorage::IndexHash Index) const override { return Implementation.FindIndexedRow(Index); }
+	TypedElementDataStorage::FQueryResult RunQuery(TypedElementDataStorage::QueryHandle Query) override { return Implementation.RunQuery(Query); }
+	TypedElementDataStorage::FQueryResult RunSubquery(int32 SubqueryIndex) override { return Implementation.RunSubquery(SubqueryIndex); }
+	TypedElementDataStorage::FQueryResult RunSubquery(int32 SubqueryIndex, TypedElementDataStorage::SubqueryCallbackRef Callback) override { return Implementation.RunSubquery(SubqueryIndex, Callback); }
+	TypedElementDataStorage::FQueryResult RunSubquery(int32 SubqueryIndex, TypedElementDataStorage::RowHandle Row, TypedElementDataStorage::SubqueryCallbackRef Callback) override { return Implementation.RunSubquery(SubqueryIndex, Row, Callback); }
+	
+	FMassQueryContextImplementation Implementation;
+};
 
 /**
  * FPhasePreOrPostAmbleExecutor
diff --git a/Engine/Plugins/Experimental/TypedElementsDataStorage/Source/TypedElementsDataStorage/Private/TypedElementDatabase.cpp b/Engine/Plugins/Experimental/TypedElementsDataStorage/Source/TypedElementsDataStorage/Private/TypedElementDatabase.cpp
--- a/Engine/Plugins/Experimental/TypedElementsDataStorage/Source/TypedElementsDataStorage/Private/TypedElementDatabase.cpp
+++ b/Engine/Plugins/Experimental/TypedElementsDataStorage/Source/TypedElementsDataStorage/Private/TypedElementDatabase.cpp
@@ -16,6 +16,7 @@
 #include "ProfilingDebugging/CpuProfilerTrace.h"
 #include "Stats/Stats2.h"
 #include "TickTaskManagerInterface.h"
+#include "TypedElementDatabaseEnvironment.h"
 #include "UObject/UObjectIterator.h"
 
 const FName UTypedElementDatabase::TickGroupName_Default(TEXT("Default"));
@@ -971,6 +972,16 @@ TSharedPtr<const FTypedElementDatabaseEnvironment> UTypedElementDatabase::GetEnv
 	return Environment;
 }
 
+FMassArchetypeHandle UTypedElementDatabase::LookupArchetype(TypedElementDataStorage::TableHandle InTableHandle) const
+{
+	const uint32 TableIndex = InTableHandle;
+	if (Tables.IsValidIndex(TableIndex))
+	{
+		return Tables[TableIndex];
+	}
+	return FMassArchetypeHandle();
+}
+
 void UTypedElementDatabase::DebugPrintQueryCallbacks(FOutputDevice& Output)
 {
 	Environment->GetQueryStore().DebugPrintQueryCallbacks(Output);
diff --git a/Engine/Plugins/Experimental/TypedElementsDataStorage/Source/TypedElementsDataStorage/Private/TypedElementDatabase.h b/Engine/Plugins/Experimental/TypedElementsDataStorage/Source/TypedElementsDataStorage/Private/TypedElementDatabase.h
--- a/Engine/Plugins/Experimental/TypedElementsDataStorage/Source/TypedElementsDataStorage/Private/TypedElementDatabase.h
+++ b/Engine/Plugins/Experimental/TypedElementsDataStorage/Source/TypedElementsDataStorage/Private/TypedElementDatabase.h
@@ -8,7 +8,6 @@
 #include "Misc/TVariant.h"
 #include "Templates/SharedPointer.h"
 #include "TypedElementDatabaseCommandBuffer.h"
-#include "TypedElementDatabaseEnvironment.h"
 #include "UObject/ObjectMacros.h"
 #include "UObject/StrongObjectPtr.h"
 
@@ -16,10 +15,12 @@
 
 struct FMassEntityManager;
 struct FMassProcessingPhaseManager;
+class FTypedElementDatabaseEnvironment;
 class UTypedElementDataStorageFactory;
 class FOutputDevice;
 class UWorld;
 
+
 UCLASS()
 class TYPEDELEMENTSDATASTORAGE_API UTypedElementDatabase 
 	: public UObject
@@ -139,6 +140,8 @@ public:
 	TSharedPtr<FTypedElementDatabaseEnvironment> GetEnvironment();
 	TSharedPtr<const FTypedElementDatabaseEnvironment> GetEnvironment() const;
 
+	FMassArchetypeHandle LookupArchetype(TypedElementDataStorage::TableHandle InTableHandle) const;
+
 	void DebugPrintQueryCallbacks(FOutputDevice& Output);
 
 	static void AddReferencedObjects(UObject* InThis, FReferenceCollector& Collector);
diff --git a/Engine/Plugins/Experimental/TypedElementsDataStorage/Source/TypedElementsDataStorage/Private/TypedElementDatabaseEnvironment.cpp b/Engine/Plugins/Experimental/TypedElementsDataStorage/Source/TypedElementsDataStorage/Private/TypedElementDatabaseEnvironment.cpp
--- a/Engine/Plugins/Experimental/TypedElementsDataStorage/Source/TypedElementsDataStorage/Private/TypedElementDatabaseEnvironment.cpp
+++ b/Engine/Plugins/Experimental/TypedElementsDataStorage/Source/TypedElementsDataStorage/Private/TypedElementDatabaseEnvironment.cpp
@@ -1,11 +1,13 @@
 // Copyright Epic Games, Inc. All Rights Reserved.
 
 #include "TypedElementDatabaseEnvironment.h"
+#include "TypedElementDatabase.h"
 
-FTypedElementDatabaseEnvironment::FTypedElementDatabaseEnvironment(ITypedElementDataStorageInterface& DataStorage,
+FTypedElementDatabaseEnvironment::FTypedElementDatabaseEnvironment(UTypedElementDatabase& InDataStorage,
 	FMassEntityManager& InMassEntityManager, FMassProcessingPhaseManager& InMassPhaseManager)
-	: DirectDeferredCommands(*this)
-	, MementoSystem(DataStorage)
+	: DataStorage(InDataStorage)
+	, DirectDeferredCommands(*this)
+	, MementoSystem(InDataStorage)
 	, MassEntityManager(InMassEntityManager)
 	, MassPhaseManager(InMassPhaseManager)
 {
@@ -71,6 +73,11 @@ const FMassEntityManager& FTypedElementDatabaseEnvironment::GetMassEntityManager
 	return MassEntityManager;
 }
 
+FMassArchetypeHandle FTypedElementDatabaseEnvironment::LookupMassArchetype(TypedElementDataStorage::TableHandle TableHandle) const
+{
+	return DataStorage.LookupArchetype(TableHandle);
+}
+
 FMassProcessingPhaseManager& FTypedElementDatabaseEnvironment::GetMassPhaseManager()
 {
 	return MassPhaseManager;
diff --git a/Engine/Plugins/Experimental/TypedElementsDataStorage/Source/TypedElementsDataStorage/Private/TypedElementDatabaseEnvironment.h b/Engine/Plugins/Experimental/TypedElementsDataStorage/Source/TypedElementsDataStorage/Private/TypedElementDatabaseEnvironment.h
--- a/Engine/Plugins/Experimental/TypedElementsDataStorage/Source/TypedElementsDataStorage/Private/TypedElementDatabaseEnvironment.h
+++ b/Engine/Plugins/Experimental/TypedElementsDataStorage/Source/TypedElementsDataStorage/Private/TypedElementDatabaseEnvironment.h
@@ -10,10 +10,12 @@
 #include "Memento/TypedElementMementoSystem.h"
 #include "Queries/TypedElementExtendedQueryStore.h"
 
+class UTypedElementDatabase;
+
 class FTypedElementDatabaseEnvironment final
 {
 public:
-	FTypedElementDatabaseEnvironment(ITypedElementDataStorageInterface& DataStorage, 
+	FTypedElementDatabaseEnvironment(UTypedElementDatabase& InDataStorage, 
 		FMassEntityManager& InMassEntityManager, FMassProcessingPhaseManager& InMassPhaseManager);
 
 	FTypedElementDatabaseCommandBuffer& GetDirectDeferredCommands();
@@ -33,6 +35,8 @@ public:
 
 	FMassEntityManager& GetMassEntityManager();
 	const FMassEntityManager& GetMassEntityManager() const;
+
+	FMassArchetypeHandle LookupMassArchetype(TypedElementDataStorage::TableHandle TableHandle) const;
 	
 	FMassProcessingPhaseManager& GetMassPhaseManager();
 	const FMassProcessingPhaseManager& GetMassPhaseManager() const;
@@ -41,6 +45,7 @@ public:
 	uint64 GetUpdateCycleId() const;
 
 private:
+	UTypedElementDatabase& DataStorage;
 	FTypedElementDatabaseCommandBuffer DirectDeferredCommands;
 	FTypedElementDatabaseIndexTable IndexTable;
 	FTypedElementDatabaseScratchBuffer ScratchBuffer;
diff --git a/Engine/Source/Runtime/TypedElementFramework/Public/Elements/Interfaces/TypedElementQueryStorageInterfaces.h b/Engine/Source/Runtime/TypedElementFramework/Public/Elements/Interfaces/TypedElementQueryStorageInterfaces.h
--- a/Engine/Source/Runtime/TypedElementFramework/Public/Elements/Interfaces/TypedElementQueryStorageInterfaces.h
+++ b/Engine/Source/Runtime/TypedElementFramework/Public/Elements/Interfaces/TypedElementQueryStorageInterfaces.h
@@ -100,6 +100,11 @@ namespace TypedElementDataStorage
 		 * in avoiding duplicated work.
 		 */
 		virtual uint64 GetUpdateCycleId() const = 0;
+		
+		/** Checks whether or not a row is in use. This is true even if the row has only been reserved. */
+		virtual bool IsRowAvailable(RowHandle Row) const = 0;
+		/** Checks whether or not a row has been reserved but not yet assigned to a table. */
+		virtual bool HasRowBeenAssigned(RowHandle Row) const = 0;
 
 		/**
 		 * Triggers all queries registered under the activation name to run for one update cycle. The activatable queries will be activated at
@@ -108,6 +113,22 @@ namespace TypedElementDataStorage
 		 */
 		virtual void ActivateQueries(FName ActivationName) = 0;
 
+		/**
+		 * Adds a row to the given table.
+		 */
+		virtual RowHandle AddRow(TableHandle Table) = 0;
+		
+		/**
+		 * Removes the row with the provided row handle. The removal will not be immediately done but delayed until the end of the tick
+		 * group.
+		 */
+		virtual void RemoveRow(RowHandle Row) = 0;
+		/**
+		 * Removes rows with the provided row handles. The removal will not be immediately done but delayed until the end of the tick
+		 * group.
+		 */
+		virtual void RemoveRows(TConstArrayView<RowHandle> Rows) = 0;
+
 		/**
 		 * Adds the provided column to the requested row.
 		 *
@@ -216,21 +237,6 @@ namespace TypedElementDataStorage
 		virtual void GetDependencies(TArrayView<UObject*> RetrievedAddresses, TConstArrayView<TWeakObjectPtr<const UClass>> DependencyTypes,
 			TConstArrayView<EQueryAccessType> AccessTypes) = 0;
 
-		/** Checks whether or not a row is in use. This is true even if the row has only been reserved. */
-		virtual bool IsRowAvailable(RowHandle Row) const = 0;
-		/** Checks whether or not a row has been reserved but not yet assigned to a table. */
-		virtual bool HasRowBeenAssigned(RowHandle Row) const = 0;
-		/**
-		 * Removes the row with the provided row handle. The removal will not be immediately done but delayed until the end of the tick
-		 * group.
-		 */
-		virtual void RemoveRow(RowHandle Row) = 0;
-		/**
-		 * Removes rows with the provided row handles. The removal will not be immediately done but delayed until the end of the tick
-		 * group.
-		 */
-		virtual void RemoveRows(TConstArrayView<RowHandle> Rows) = 0;
-
 		/** Retrieves the row for an indexed object. Returns an invalid row handle if the hash wasn't found. */
 		virtual RowHandle FindIndexedRow(IndexHash Index) const = 0;
 
diff --git a/Engine/Source/Runtime/TypedElementFramework/Tests/Elements/Framework/TypedElementDataStorageProcessorTests.cpp b/Engine/Source/Runtime/TypedElementFramework/Tests/Elements/Framework/TypedElementDataStorageProcessorTests.cpp
new file mode 100644
--- /dev/null
+++ b/Engine/Source/Runtime/TypedElementFramework/Tests/Elements/Framework/TypedElementDataStorageProcessorTests.cpp
@@ -0,0 +1,183 @@
+﻿// Copyright Epic Games, Inc. All Rights Reserved.
+
+#if WITH_TESTS
+#include "Elements/Framework/TypedElementRegistry.h"
+#include "HAL/IConsoleManager.h"
+#include "TypedElementTestColumns.h"
+
+namespace TEDS::Debug::ProcessorTests
+{
+	void OnProcessorTestsEnabled(IConsoleVariable*);
+
+	bool GProcessorTestsEnabled = false;
+	FAutoConsoleVariableRef CVarAllowUnversionedContentInEditor(
+		TEXT("TEDS.Tests.ProcessorTestsEnabled"),
+		GProcessorTestsEnabled,
+		TEXT("If true, registers processors and additional commands with TEDS to test processors."),
+		FConsoleVariableDelegate::CreateStatic(&OnProcessorTestsEnabled),
+		ECVF_Default
+	);
+
+	TypedElementDataStorage::TableHandle PrimaryTable = TypedElementDataStorage::InvalidTableHandle;
+	TypedElementDataStorage::TableHandle SecondaryTable = TypedElementDataStorage::InvalidTableHandle;
+
+	TArray<TypedElementQueryHandle> RegisteredQueries;
+	TArray<IConsoleCommand*> RegisteredCommands;
+
+	void RegisterProcessors()
+	{
+		ITypedElementDataStorageInterface* DataStorage = UTypedElementRegistry::GetInstance()->GetMutableDataStorage();
+		using namespace TypedElementQueryBuilder;
+		using DSI = ITypedElementDataStorageInterface;
+		namespace DS = TypedElementDataStorage;
+
+		if (PrimaryTable == TypedElementDataStorage::InvalidTableHandle)
+		{
+			PrimaryTable =
+				DataStorage->RegisterTable<FTEDSProcessorTestsReferenceColumn, FTEDSProcessorTests_PrimaryTag>(FName(TEXT("ProcessorTests Primary Table")));
+		}
+		if (SecondaryTable == TypedElementDataStorage::InvalidTableHandle)
+		{
+			SecondaryTable =
+				DataStorage->RegisterTable<FTEDSProcessorTestsReferenceColumn, FTEDSProcessorTests_SecondaryTag>(FName(TEXT("ProcessorTests Secondary Table")));
+		}
+
+		// Test creation of a row from within a query processor
+		TypedElementQueryHandle PrimaryRowQuery = DataStorage->RegisterQuery(
+			Select(TEXT("TEST: Creating a row for primary reference column"),
+			FProcessor(DSI::EQueryTickPhase::PrePhysics, DataStorage->GetQueryTickGroupName(DSI::EQueryTickGroups::Default)),
+			[](DS::IQueryContext& Context, const TypedElementRowHandle* Rows, FTEDSProcessorTestsReferenceColumn* ReferenceColumns)
+			{
+				const int32 RowCount = Context.GetRowCount();
+				TConstArrayView<TypedElementRowHandle> RowsView = MakeArrayView(Rows, RowCount);
+				TArrayView<FTEDSProcessorTestsReferenceColumn> ReferenceColumnsView = MakeArrayView(ReferenceColumns, RowCount);
+
+				for (int32 Index = 0; Index < RowCount; ++Index)
+				{
+					FTEDSProcessorTestsReferenceColumn& PrimaryReferenceColumn = ReferenceColumnsView[Index];
+					// Auto-create a secondary row if this points to no row
+					if (!Context.IsRowAvailable(PrimaryReferenceColumn.Reference))
+					{
+						TypedElementDataStorage::RowHandle SecondaryRow = Context.AddRow(SecondaryTable);
+						// Initialize bi-directional row references
+						PrimaryReferenceColumn.Reference = SecondaryRow;
+						Context.AddColumn(SecondaryRow, FTEDSProcessorTestsReferenceColumn{.Reference = RowsView[Index]});
+					}
+				}
+			}
+			)
+			.Where()
+				.All<FTEDSProcessorTests_PrimaryTag>()
+				.None<FTEDSProcessorTests_Linked>()
+			.Compile());
+
+		RegisteredQueries.Emplace(PrimaryRowQuery);
+
+		TypedElementQueryHandle UpdateTransformWidget = DataStorage->RegisterQuery(
+			Select()
+				.ReadOnly<FTEDSProcessorTestsReferenceColumn>()
+			.Where()
+				.All<FTEDSProcessorTests_PrimaryTag>()
+				.None<FTEDSProcessorTests_Linked>()
+			.Compile());
+		
+		RegisteredQueries.Emplace(UpdateTransformWidget);
+
+		TypedElementQueryHandle SecondaryRowQuery = DataStorage->RegisterQuery(
+			Select(TEXT("TEST: Creating a row for secondary reference column"),
+			FProcessor(DSI::EQueryTickPhase::DuringPhysics, DataStorage->GetQueryTickGroupName(DSI::EQueryTickGroups::Default)),
+			[](DS::IQueryContext& Context, const TypedElementRowHandle* Rows, FTEDSProcessorTestsReferenceColumn* ReferenceColumns)
+			{
+				const int32 RowCount = Context.GetRowCount();
+				TConstArrayView<TypedElementRowHandle> RowsView = MakeArrayView(Rows, RowCount);
+				TArrayView<FTEDSProcessorTestsReferenceColumn> ReferenceColumnsView = MakeArrayView(ReferenceColumns, RowCount);
+
+				for (int32 Index = 0; Index < RowCount; ++Index)
+				{
+					const TypedElementRowHandle& SecondaryRow = RowsView[Index];
+					Context.RunSubquery(
+						0,
+						ReferenceColumnsView[Index].Reference,
+						CreateSubqueryCallbackBinding(
+						[SecondaryRow](TypedElementDataStorage::ISubqueryContext& SubqueryContext, TypedElementDataStorage::RowHandle PrimaryRow, const FTEDSProcessorTestsReferenceColumn& ReferenceColumn)
+						{
+							if (ReferenceColumn.Reference == SecondaryRow)
+							{
+								// Add these tags to prevent further processing
+								SubqueryContext.AddColumns<FTEDSProcessorTests_Linked>(SecondaryRow);
+								SubqueryContext.AddColumns<FTEDSProcessorTests_Linked>(PrimaryRow);
+							}
+						}
+					));
+				}
+			}
+			)
+			.Where()
+				.All<FTEDSProcessorTests_SecondaryTag>()
+				.None<FTEDSProcessorTests_Linked>()
+			.DependsOn()
+				.SubQuery(UpdateTransformWidget)
+			.Compile());
+
+		RegisteredQueries.Emplace(SecondaryRowQuery);
+	}
+
+	void UnregisterProcessors()
+	{
+		ITypedElementDataStorageInterface* DataStorage = UTypedElementRegistry::GetInstance()->GetMutableDataStorage();
+		for (const TypedElementQueryHandle& Handle : RegisteredQueries)
+		{
+			DataStorage->UnregisterQuery(Handle);
+		}
+		RegisteredQueries.Empty();
+	}
+
+	void RegisterCommands()
+	{
+		IConsoleCommand* Command = IConsoleManager::Get().RegisterConsoleCommand(TEXT("TEDS.Tests.ProcessorTests.AddPrimaryRows"), TEXT(""), FConsoleCommandWithArgsDelegate::CreateLambda([](const TArray<FString>& Args)
+		{
+			int32 RowsToCreate = 0;
+			if (Args.Num() != 1)
+			{
+				return;
+			}
+			TTypeFromString<int32>::FromString(RowsToCreate, *Args[0]);
+			if (RowsToCreate <= 0)
+			{
+				return;
+			}
+
+			ITypedElementDataStorageInterface* DataStorage = UTypedElementRegistry::GetInstance()->GetMutableDataStorage();
+			DataStorage->BatchAddRow(PrimaryTable, RowsToCreate, [](TypedElementRowHandle Row)
+			{
+			});
+		}), ECVF_Default);
+		
+		RegisteredCommands.Emplace(Command);
+	}
+
+	void UnregisterCommands()
+	{
+		for (IConsoleCommand* Command : RegisteredCommands)
+		{
+			IConsoleManager::Get().UnregisterConsoleObject(Command);
+		}
+		RegisteredCommands.Empty();
+	}
+	
+	void OnProcessorTestsEnabled(IConsoleVariable* Variable)
+	{
+		if (Variable->GetBool())
+		{
+			RegisterProcessors();
+			RegisterCommands();
+		}
+		else
+		{
+			UnregisterProcessors();
+			UnregisterCommands();
+		}
+	};
+}
+#endif
+
diff --git a/Engine/Source/Runtime/TypedElementFramework/Tests/Elements/Framework/TypedElementTestColumns.h b/Engine/Source/Runtime/TypedElementFramework/Tests/Elements/Framework/TypedElementTestColumns.h
--- a/Engine/Source/Runtime/TypedElementFramework/Tests/Elements/Framework/TypedElementTestColumns.h
+++ b/Engine/Source/Runtime/TypedElementFramework/Tests/Elements/Framework/TypedElementTestColumns.h
@@ -69,6 +69,35 @@ struct FTestTagColumnC final : public FTypedElementDataStorageTag
 
 USTRUCT(meta = (DisplayName = "TagD"))
 struct FTestTagColumnD final : public FTypedElementDataStorageTag
+{
+	GENERATED_BODY()
+};
+
+USTRUCT(meta = (DisplayName = "TestReferenceColumn"))
+struct FTEDSProcessorTestsReferenceColumn final : public FTypedElementDataStorageColumn
+{
+	GENERATED_BODY()
+
+	// UPROPERTY()
+	TypedElementDataStorage::RowHandle Reference = TypedElementDataStorage::InvalidRowHandle;
+
+	bool IsReferenced = false;
+};
+
+USTRUCT(meta = (DisplayName = "TEDSProcessorTests_PrimaryTag"))
+struct FTEDSProcessorTests_PrimaryTag final : public FTypedElementDataStorageTag
+{
+	GENERATED_BODY()
+};
+
+USTRUCT(meta = (DisplayName = "TEDSProcessorTestsSecondaryTag"))
+struct FTEDSProcessorTests_SecondaryTag final : public FTypedElementDataStorageTag
+{
+	GENERATED_BODY()
+};
+
+USTRUCT(meta = (DisplayName = "TEDSProcessorTests_Linked"))
+struct FTEDSProcessorTests_Linked final : public FTypedElementDataStorageTag
 {
 	GENERATED_BODY()
 };
\ No newline at end of file
