Fix CIS due to missing #include

diff --git a/Engine/Plugins/Experimental/TypedElementsDataStorage/Source/TypedElementsDataStorage/Private/TypedElementDatabaseCompatibility.cpp b/Engine/Plugins/Experimental/TypedElementsDataStorage/Source/TypedElementsDataStorage/Private/TypedElementDatabaseCompatibility.cpp
--- a/Engine/Plugins/Experimental/TypedElementsDataStorage/Source/TypedElementsDataStorage/Private/TypedElementDatabaseCompatibility.cpp
+++ b/Engine/Plugins/Experimental/TypedElementsDataStorage/Source/TypedElementsDataStorage/Private/TypedElementDatabaseCompatibility.cpp
@@ -17,6 +17,7 @@
 #include "Elements/Framework/TypedElementQueryBuilder.h"
 #include "Memento/TypedElementMementoRowTypes.h"
 #include "TypedElementDatabase.h"
+#include "TypedElementDatabaseEnvironment.h"
 #include "TypedElementDataStorageProfilingMacros.h"
 
 void UTypedElementDatabaseCompatibility::Initialize(UTypedElementDatabase* InStorage)
