Added editor settings for TEDS' Mass backend.

The Mass backend for TEDS can now be configured in the editor preferences under the advanced sections. The first option currently only support settings the chunk size for tables.

#rnx

diff --git a/Engine/Plugins/Experimental/TypedElementsDataStorage/Source/TypedElementsDataStorage/Private/EditorDataStorageSettings.h b/Engine/Plugins/Experimental/TypedElementsDataStorage/Source/TypedElementsDataStorage/Private/EditorDataStorageSettings.h
new file mode 100644
--- /dev/null
+++ b/Engine/Plugins/Experimental/TypedElementsDataStorage/Source/TypedElementsDataStorage/Private/EditorDataStorageSettings.h
@@ -0,0 +1,34 @@
+// Copyright Epic Games, Inc. All Rights Reserved.
+
+#pragma once
+
+#include "CoreMinimal.h"
+#include "UObject/ObjectMacros.h"
+#include "UObject/Object.h"
+#include "EditorDataStorageSettings.generated.h"
+
+UENUM()
+enum class EChunkMemorySize : int32
+{
+	Size4Kb = 4 * 1024 UMETA(DisplayName = "4kb"),
+	Size8Kb = 8 * 1024 UMETA(DisplayName = "8kb"),
+	Size16Kb = 16 * 1024 UMETA(DisplayName = "16kb"),
+	Size32Kb = 32 * 1024 UMETA(DisplayName = "32kb"),
+	Size64Kb = 64 * 1024 UMETA(DisplayName = "64kb"),
+	Size128Kb = 128 * 1024 UMETA(DisplayName = "128kb"),
+	Size256Kb = 256 * 1024 UMETA(DisplayName = "256kb"),
+	Size512Kb = 512 * 1024 UMETA(DisplayName = "512kb")
+};
+
+UCLASS(config = EditorPerProjectUserSettings, MinimalAPI)
+class UEditorDataStorageSettings : public UObject
+{
+	GENERATED_BODY()
+
+public:
+	UPROPERTY(EditAnywhere, config, Category = MassSettings)
+	EChunkMemorySize ChunkMemorySize = EChunkMemorySize::Size128Kb;
+
+	UPROPERTY(EditAnywhere, config, Category = MassSettings)
+	TMap<FName, EChunkMemorySize> TableSpecificChunkMemorySize;
+};
diff --git a/Engine/Plugins/Experimental/TypedElementsDataStorage/Source/TypedElementsDataStorage/Private/TypedElementDatabase.cpp b/Engine/Plugins/Experimental/TypedElementsDataStorage/Source/TypedElementsDataStorage/Private/TypedElementDatabase.cpp
--- a/Engine/Plugins/Experimental/TypedElementsDataStorage/Source/TypedElementsDataStorage/Private/TypedElementDatabase.cpp
+++ b/Engine/Plugins/Experimental/TypedElementsDataStorage/Source/TypedElementsDataStorage/Private/TypedElementDatabase.cpp
@@ -3,6 +3,7 @@
 #include "TypedElementDatabase.h"
 
 #include "Editor.h"
+#include "EditorDataStorageSettings.h"
 #include "Elements/Interfaces/TypedElementDataStorageFactory.h"
 #include "Elements/Framework/TypedElementRegistry.h"
 #include "Engine/World.h"
@@ -269,6 +270,7 @@ TypedElementTableHandle UTypedElementDatabase::RegisterTable(TConstArrayView<con
 		TypedElementTableHandle Result = Tables.Num();
 		FMassArchetypeCreationParams ArchetypeCreationParams;
 		ArchetypeCreationParams.DebugName = Name;
+		ArchetypeCreationParams.ChunkMemorySize = GetTableChunkSize(Name);
 		Tables.Add(ActiveEditorEntityManager->CreateArchetype(ColumnList, ArchetypeCreationParams));
 		if (Name.IsValid())
 		{
@@ -287,6 +289,7 @@ TypedElementTableHandle UTypedElementDatabase::RegisterTable(TypedElementTableHa
 		TypedElementTableHandle Result = Tables.Num();
 		FMassArchetypeCreationParams ArchetypeCreationParams;
 		ArchetypeCreationParams.DebugName = Name;
+		ArchetypeCreationParams.ChunkMemorySize = GetTableChunkSize(Name);
 		Tables.Add(ActiveEditorEntityManager->CreateArchetype(Tables[SourceTable], ColumnList, ArchetypeCreationParams));
 		if (Name.IsValid())
 		{
@@ -962,6 +965,19 @@ void UTypedElementDatabase::Reset()
 	ActiveEditorEntityManager.Reset();
 }
 
+int32 UTypedElementDatabase::GetTableChunkSize(FName TableName) const
+{
+	const UEditorDataStorageSettings* Settings = GetDefault<UEditorDataStorageSettings>();
+	if (const EChunkMemorySize* TableSpecificSize = Settings->TableSpecificChunkMemorySize.Find(TableName))
+	{
+		return static_cast<int32>(*TableSpecificSize);
+	}
+	else
+	{
+		return static_cast<int32>(Settings->ChunkMemorySize);
+	}
+}
+
 TSharedPtr<FTypedElementDatabaseEnvironment> UTypedElementDatabase::GetEnvironment()
 {
 	return Environment;
diff --git a/Engine/Plugins/Experimental/TypedElementsDataStorage/Source/TypedElementsDataStorage/Private/TypedElementDatabase.h b/Engine/Plugins/Experimental/TypedElementsDataStorage/Source/TypedElementsDataStorage/Private/TypedElementDatabase.h
--- a/Engine/Plugins/Experimental/TypedElementsDataStorage/Source/TypedElementsDataStorage/Private/TypedElementDatabase.h
+++ b/Engine/Plugins/Experimental/TypedElementsDataStorage/Source/TypedElementsDataStorage/Private/TypedElementDatabase.h
@@ -150,6 +150,8 @@ private:
 	void PreparePhase(EQueryTickPhase Phase, float DeltaTime);
 	void FinalizePhase(EQueryTickPhase Phase, float DeltaTime);
 	void Reset();
+
+	int32 GetTableChunkSize(FName TableName) const;
 	
 	struct FFactoryTypePair
 	{
diff --git a/Engine/Plugins/Experimental/TypedElementsDataStorage/Source/TypedElementsDataStorage/Private/TypedElementsDataStorage.cpp b/Engine/Plugins/Experimental/TypedElementsDataStorage/Source/TypedElementsDataStorage/Private/TypedElementsDataStorage.cpp
--- a/Engine/Plugins/Experimental/TypedElementsDataStorage/Source/TypedElementsDataStorage/Private/TypedElementsDataStorage.cpp
+++ b/Engine/Plugins/Experimental/TypedElementsDataStorage/Source/TypedElementsDataStorage/Private/TypedElementsDataStorage.cpp
@@ -2,15 +2,17 @@
 
 #include "TypedElementsDataStorage.h"
 
+#include "EditorDataStorageSettings.h"
+#include "Elements/Common/TypedElementDataStorageLog.h"
 #include "Elements/Framework/TypedElementRegistry.h"
 #include "Elements/Interfaces/TypedElementDataStorageFactory.h"
+#include "ISettingsModule.h"
 #include "MassEntityTypes.h"
 #include "Misc/CoreDelegates.h"
 #include "Templates/IsPolymorphic.h"
 #include "TypedElementDatabase.h"
 #include "TypedElementDatabaseCompatibility.h"
 #include "TypedElementDatabaseUI.h"
-#include "Elements/Common/TypedElementDataStorageLog.h"
 #include "UObject/UObjectGlobals.h"
 
 #define LOCTEXT_NAMESPACE "FTypedElementsDataStorageModule"
@@ -68,7 +70,17 @@ void FTypedElementsDataStorageModule::StartupModule()
 
 	UE_LOG(LogTypedElementDataStorage, Log, TEXT("Enabled by TEDS.Enable CVar"));
 
-	
+	// Setup the editor settings;
+	if (ISettingsModule* SettingsModule = FModuleManager::GetModulePtr<ISettingsModule>("Settings"))
+	{
+		// automation
+		SettingsModule->RegisterSettings("Editor", "Advanced", "The Editor Data Storage",
+			LOCTEXT("DataStorageSettingsName", "The Editor Data Storage"),
+			LOCTEXT("DataStorageSettingsDescription", "Configuration options for the central data storage used by various tools to store their data."),
+			GetMutableDefault<UEditorDataStorageSettings>()
+		);
+	}
+
 	// Load the dependent TypedElementFramework module (holding TypedElementRegistry) here so that it is guaranteed to be available in Shutdown
 	// and it is shutdown AFTER FTypedElementsDataStorageModule
 	FModuleManager::Get().LoadModule(TEXT("TypedElementFramework"));
