Added the ability for Mass processors/observers to provide a description.

The Mass processor is updated with a virtual function that allows a description to be provided as key/value string pairs. When available the Mass Debugger will show these, otherwise the Mass Debugger will look the same as before. This change was added because TEDS has additional information that it appended to the name of the processor. This is increasingly difficult to read in the Mass Debugger and leads to other problems when needing to reference the name of a processor.

[RN] If the debug output description provided by a Mass processor now has more than the processor name it it, the description will be shown in the Mass Debugger.

diff --git a/Engine/Plugins/Experimental/TypedElementsDataStorage/Source/TypedElementsDataStorage/Private/Processors/TypedElementProcessorAdaptors.cpp b/Engine/Plugins/Experimental/TypedElementsDataStorage/Source/TypedElementsDataStorage/Private/Processors/TypedElementProcessorAdaptors.cpp
--- a/Engine/Plugins/Experimental/TypedElementsDataStorage/Source/TypedElementsDataStorage/Private/Processors/TypedElementProcessorAdaptors.cpp
+++ b/Engine/Plugins/Experimental/TypedElementsDataStorage/Source/TypedElementsDataStorage/Private/Processors/TypedElementProcessorAdaptors.cpp
@@ -674,6 +674,7 @@ void FPhasePreOrPostAmbleExecutor::ExecuteQuery(
 /**
  * FTypedElementQueryProcessorData
  */
+
 FTypedElementQueryProcessorData::FTypedElementQueryProcessorData(UMassProcessor& Owner)
 	: NativeQuery(Owner)
 {
@@ -743,22 +744,60 @@ EMassProcessingPhase FTypedElementQueryProcessorData::MapToMassProcessingPhase(I
 
 FString FTypedElementQueryProcessorData::GetProcessorName() const
 {
-	FString Result;
 	if (const FTypedElementExtendedQuery* StoredQuery = QueryStore ? QueryStore->Get(ParentQuery) : nullptr)
 	{
-		Result = StoredQuery->Description.Callback.Name.ToString();
-		if (!StoredQuery->Description.Callback.ActivationName.IsNone())
-		{
-			Result += TEXT(" (Activatable: '");
-			StoredQuery->Description.Callback.ActivationName.AppendString(Result);
-			Result += TEXT("')");
-		}
+		return StoredQuery->Description.Callback.Name.ToString();
 	}
 	else
 	{
-		Result = TEXT("<unnamed>");
+		return TEXT("<unnamed>");
 	}
-	return Result;
+}
+
+void FTypedElementQueryProcessorData::DebugOutputDescription(FOutputDevice& Ar, int32 Indent) const
+{
+#if WITH_MASSENTITY_DEBUG
+	if (const FTypedElementExtendedQuery* StoredQuery = QueryStore ? QueryStore->Get(ParentQuery) : nullptr)
+	{
+		const ITypedElementDataStorageInterface::FQueryDescription& Description = StoredQuery->Description;
+		const ITypedElementDataStorageInterface::FQueryDescription::FCallbackData& Callback = Description.Callback;
+		
+		if (!Callback.Group.IsNone())
+		{
+			Ar.Logf(TEXT("\n%*sGroup: %s"), Indent, TEXT(""), *Callback.Group.ToString());
+		}
+		if (!Callback.BeforeGroups.IsEmpty())
+		{
+			Ar.Logf(TEXT("\n%*sBefore:"), Indent, TEXT(""));
+			int32 Index = 0;
+			for (FName BeforeName : Callback.BeforeGroups)
+			{
+				Ar.Logf(TEXT("\n%*s[%i] %s"), Indent + 4, TEXT(""), Index++, *BeforeName.ToString());
+			}
+		}
+		if (!Callback.AfterGroups.IsEmpty())
+		{
+			Ar.Logf(TEXT("\n%*sAfter:"), Indent, TEXT(""));
+			int32 Index = 0;
+			for (FName AfterName : Callback.AfterGroups)
+			{
+				Ar.Logf(TEXT("\n%*s[%i] %s"), Indent + 4, TEXT(""), Index++, *AfterName.ToString());
+			}
+		}
+		
+		if (!Callback.ActivationName.IsNone())
+		{
+			Ar.Logf(TEXT("\n%*sActivatable: %s"), Indent, TEXT(""), *Callback.ActivationName.ToString());
+		}
+
+		if (Callback.MonitoredType)
+		{
+			Ar.Logf(TEXT("\n%*sMonitored type: %s"), Indent, TEXT(""), *Callback.MonitoredType->GetName());
+		}
+
+		Ar.Logf(TEXT("\n%*sIs forced to GameThread: %s"), Indent, TEXT(""), Callback.bForceToGameThread ? TEXT("True") : TEXT("False"));
+	}
+#endif // WITH_MASSENTITY_DEBUG
 }
 
 bool FTypedElementQueryProcessorData::PrepareCachedDependenciesOnQuery(
@@ -975,9 +1014,16 @@ void UTypedElementQueryProcessorCallbackAdapterProcessorBase::PostInitProperties
 
 FString UTypedElementQueryProcessorCallbackAdapterProcessorBase::GetProcessorName() const
 {
-	FString Name = Data.GetProcessorName();
-	Name += TEXT(" [Editor Processor]");
-	return Name;
+	return Data.GetProcessorName();
+}
+
+void UTypedElementQueryProcessorCallbackAdapterProcessorBase::DebugOutputDescription(FOutputDevice& Ar, int32 Indent) const
+{
+#if WITH_MASSENTITY_DEBUG
+	UMassProcessor::DebugOutputDescription(Ar, Indent);
+	Ar.Logf(TEXT("\n%*sType: Editor Processor"), Indent, TEXT(""));
+	Data.DebugOutputDescription(Ar, Indent);
+#endif // WITH_MASSENTITY_DEBUG
 }
 
 void UTypedElementQueryProcessorCallbackAdapterProcessorBase::Execute(FMassEntityManager& EntityManager, FMassExecutionContext& Context)
@@ -1114,23 +1160,31 @@ void UTypedElementQueryObserverCallbackAdapterProcessorBase::Register()
 
 FString UTypedElementQueryObserverCallbackAdapterProcessorBase::GetProcessorName() const
 {
-	FString Name = Data.GetProcessorName();
+	return Data.GetProcessorName();
+}
+
+void UTypedElementQueryObserverCallbackAdapterProcessorBase::DebugOutputDescription(FOutputDevice& Ar, int32 Indent) const
+{
+#if WITH_MASSENTITY_DEBUG
+	UMassObserverProcessor::DebugOutputDescription(Ar, Indent);
 	EMassObservedOperation ObservationType = GetObservedOperation();
 	if (ObservationType == EMassObservedOperation::Add)
 	{
-		Name += TEXT(" [Editor Add Observer]");
+		Ar.Logf(TEXT("\n%*sType: Editor Add Observer"), Indent, TEXT(""));
 	}
 	else if (ObservationType == EMassObservedOperation::Remove)
 	{
-		Name += TEXT(" [Editor Remove Observer]");
+		Ar.Logf(TEXT("\n%*sType: Editor Remove Observer"), Indent, TEXT(""));
 	}
 	else
 	{
-		Name += TEXT(" [Editor <Unknown> Observer]");
+		Ar.Logf(TEXT("\n%*sType: Editor <Unknown> Observer"), Indent, TEXT(""));
 	}
-	return Name;
+	Data.DebugOutputDescription(Ar, Indent);
+#endif // WITH_MASSENTITY_DEBUG
 }
 
+
 void UTypedElementQueryObserverCallbackAdapterProcessorBase::Execute(FMassEntityManager& EntityManager, FMassExecutionContext& Context)
 {
 	Data.Execute(EntityManager, Context);
diff --git a/Engine/Plugins/Experimental/TypedElementsDataStorage/Source/TypedElementsDataStorage/Private/Processors/TypedElementProcessorAdaptors.h b/Engine/Plugins/Experimental/TypedElementsDataStorage/Source/TypedElementsDataStorage/Private/Processors/TypedElementProcessorAdaptors.h
--- a/Engine/Plugins/Experimental/TypedElementsDataStorage/Source/TypedElementsDataStorage/Private/Processors/TypedElementProcessorAdaptors.h
+++ b/Engine/Plugins/Experimental/TypedElementsDataStorage/Source/TypedElementsDataStorage/Private/Processors/TypedElementProcessorAdaptors.h
@@ -48,6 +48,7 @@ struct FTypedElementQueryProcessorData
 		TArrayView<FMassEntityQuery> Subqueries);
 	static EMassProcessingPhase MapToMassProcessingPhase(ITypedElementDataStorageInterface::EQueryTickPhase Phase);
 	FString GetProcessorName() const;
+	void DebugOutputDescription(FOutputDevice& Ar, int32 Indent) const;
 
 	static TypedElementDataStorage::FQueryResult Execute(
 		TypedElementDataStorage::DirectQueryCallbackRef& Callback,
@@ -75,7 +76,6 @@ struct FTypedElementQueryProcessorData
 	static bool PrepareCachedDependenciesOnQuery(
 		ITypedElementDataStorageInterface::FQueryDescription& Description, FMassExecutionContext& Context);
 
-	
 	FTypedElementExtendedQueryStore::Handle ParentQuery;
 	FTypedElementExtendedQueryStore* QueryStore{ nullptr };
 	FTypedElementDatabaseEnvironment* Environment{ nullptr };
@@ -114,6 +114,7 @@ protected:
 	
 	void PostInitProperties() override;
 	FString GetProcessorName() const override;
+	void DebugOutputDescription(FOutputDevice& Ar, int32 Indent) const override;
 
 private:
 	UPROPERTY(transient)
@@ -296,6 +297,7 @@ protected:
 	void PostInitProperties() override;
 	void Register() override;
 	FString GetProcessorName() const override;
+	void DebugOutputDescription(FOutputDevice& Ar, int32 Indent) const override;
 
 private:
 	UPROPERTY(transient)
diff --git a/Engine/Plugins/Runtime/MassEntity/Source/MassEntity/Public/MassProcessor.h b/Engine/Plugins/Runtime/MassEntity/Source/MassEntity/Public/MassProcessor.h
--- a/Engine/Plugins/Runtime/MassEntity/Source/MassEntity/Public/MassProcessor.h
+++ b/Engine/Plugins/Runtime/MassEntity/Source/MassEntity/Public/MassProcessor.h
@@ -66,9 +66,10 @@ public:
 	UE_DEPRECATED(5.2, "This function is deprecated. Use ShouldAllowDuplicates instead.")
 	bool AllowDuplicates() const { return ShouldAllowMultipleInstances(); }
 
-	virtual void DebugOutputDescription(FOutputDevice& Ar, int32 Indent = 0) const;
+	void DebugOutputDescription(FOutputDevice& Ar) const { DebugOutputDescription(Ar, 0); }
+	virtual void DebugOutputDescription(FOutputDevice& Ar, int32 Indent) const;
 	virtual FString GetProcessorName() const { return GetName(); }
-
+	
 	//----------------------------------------------------------------------//
 	// Ordering functions 
 	//----------------------------------------------------------------------//
diff --git a/Engine/Plugins/Runtime/MassEntity/Source/MassEntityDebugger/Private/MassDebuggerModel.cpp b/Engine/Plugins/Runtime/MassEntity/Source/MassEntityDebugger/Private/MassDebuggerModel.cpp
--- a/Engine/Plugins/Runtime/MassEntity/Source/MassEntityDebugger/Private/MassDebuggerModel.cpp
+++ b/Engine/Plugins/Runtime/MassEntity/Source/MassEntityDebugger/Private/MassDebuggerModel.cpp
@@ -167,6 +167,15 @@ void FMassDebuggerProcessorData::SetProcessor(const UMassProcessor& InProcessor)
 	UE::Mass::Debugger::Private::MakeDisplayName(Name, Label);
 
 	ProcessorHash = UE::Mass::Debugger::Private::CalcProcessorHash(InProcessor);
+
+#if WITH_MASSENTITY_DEBUG
+	FStringOutputDevice DescriptionDevice;
+	InProcessor.DebugOutputDescription(DescriptionDevice);
+	if (DescriptionDevice != InProcessor.GetProcessorName())
+	{
+		Description = MoveTemp(DescriptionDevice);
+	}
+#endif // WITH_MASSENTITY_DEBUG
 }
 
 //----------------------------------------------------------------------//
diff --git a/Engine/Plugins/Runtime/MassEntity/Source/MassEntityDebugger/Private/MassDebuggerModel.h b/Engine/Plugins/Runtime/MassEntity/Source/MassEntityDebugger/Private/MassDebuggerModel.h
--- a/Engine/Plugins/Runtime/MassEntity/Source/MassEntityDebugger/Private/MassDebuggerModel.h
+++ b/Engine/Plugins/Runtime/MassEntity/Source/MassEntityDebugger/Private/MassDebuggerModel.h
@@ -107,7 +107,9 @@ public:
 	TSharedPtr<FMassDebuggerQueryData> ProcessorRequirements;
 	TArray<TSharedPtr<FMassDebuggerQueryData>> Queries;
 	TArray<TSharedPtr<FMassDebuggerArchetypeData>> ValidArchetypes;
-	
+#if WITH_MASSENTITY_DEBUG
+	FString Description;
+#endif // WITH_MASSENTITY_DEBUG	
 	EMassDebuggerProcessorSelection Selection = EMassDebuggerProcessorSelection::None;
 };
 
diff --git a/Engine/Plugins/Runtime/MassEntity/Source/MassEntityDebugger/Private/SMassProcessor.cpp b/Engine/Plugins/Runtime/MassEntity/Source/MassEntityDebugger/Private/SMassProcessor.cpp
--- a/Engine/Plugins/Runtime/MassEntity/Source/MassEntityDebugger/Private/SMassProcessor.cpp
+++ b/Engine/Plugins/Runtime/MassEntity/Source/MassEntityDebugger/Private/SMassProcessor.cpp
@@ -36,6 +36,22 @@ void SMassProcessor::Construct(const FArguments& InArgs, TSharedPtr<FMassDebugge
 			.TextStyle(FAppStyle::Get(), "LargeText")
 		]
 	];
+
+#if WITH_MASSENTITY_DEBUG
+	if (!ProcessorData->Description.IsEmpty())
+	{
+		Box->AddSlot()
+		.AutoHeight()
+		[
+			SNew(SBorder)
+			.Padding(10.0f)
+			[
+				SNew(STextBlock)
+				.Text(FText::FromString(ProcessorData->Description))
+			]
+		];
+	}
+#endif //WITH_MASSENTITY_DEBUG
 	
 	if (ProcessorData->ProcessorRequirements->IsEmpty())
 	{
