[ALT]
 - Move TEDS registration from PostInitProperties to PostLoad during async loading

#rnx

diff --git a/Engine/Plugins/Experimental/TypedElementsDataStorage/Source/TypedElementsDataStorage/Private/TypedElementDatabaseCompatibility.cpp b/Engine/Plugins/Experimental/TypedElementsDataStorage/Source/TypedElementsDataStorage/Private/TypedElementDatabaseCompatibility.cpp
--- a/Engine/Plugins/Experimental/TypedElementsDataStorage/Source/TypedElementsDataStorage/Private/TypedElementDatabaseCompatibility.cpp
+++ b/Engine/Plugins/Experimental/TypedElementsDataStorage/Source/TypedElementsDataStorage/Private/TypedElementDatabaseCompatibility.cpp
@@ -262,6 +262,8 @@ TypedElementDataStorage::TableHandle UTypedElementDatabaseCompatibility::FindBes
 template<bool bEnableTransactions>
 TypedElementRowHandle UTypedElementDatabaseCompatibility::AddCompatibleObjectExplicitTransactionable(UObject* Object)
 {
+	check(IsInGameThread());
+
 	using namespace TypedElementDataStorage;
 
 	TypedElementRowHandle Result = FindRowWithCompatibleObjectExplicit(Object);
diff --git a/Engine/Source/Runtime/Engine/Private/Actor.cpp b/Engine/Source/Runtime/Engine/Private/Actor.cpp
--- a/Engine/Source/Runtime/Engine/Private/Actor.cpp
+++ b/Engine/Source/Runtime/Engine/Private/Actor.cpp
@@ -408,7 +408,15 @@ void AActor::PostInitProperties()
 	Super::PostInitProperties();
 
 #if WITH_EDITOR
-	UEngineElementsLibrary::CreateEditorActorElement(this);
+	if (IsInGameThread())
+	{
+		// This will be executed in PostLoad instead during loading
+		UEngineElementsLibrary::CreateEditorActorElement(this);
+	}
+	else
+	{
+		check(IsInAsyncLoadingThread());
+	}
 #endif	// WITH_EDITOR
 
 	RemoteRole = (bReplicates ? ROLE_SimulatedProxy : ROLE_None);
@@ -936,6 +944,10 @@ void AActor::PostLoad()
 {
 	Super::PostLoad();
 
+#if WITH_EDITOR
+	UEngineElementsLibrary::CreateEditorActorElement(this);
+#endif
+
 	// add ourselves to our Owner's Children array
 	if (Owner != nullptr)
 	{
diff --git a/Engine/Source/Runtime/Engine/Private/Elements/Framework/EngineElementsLibrary.cpp b/Engine/Source/Runtime/Engine/Private/Elements/Framework/EngineElementsLibrary.cpp
--- a/Engine/Source/Runtime/Engine/Private/Elements/Framework/EngineElementsLibrary.cpp
+++ b/Engine/Source/Runtime/Engine/Private/Elements/Framework/EngineElementsLibrary.cpp
@@ -45,7 +45,7 @@ void CreateEditorTypedElement(const ObjectClass* Object, TTypedElementOwnerStore
 {
 	if (GIsEditor && !Object->HasAnyFlags(RF_ClassDefaultObject | RF_ArchetypeObject))
 	{
-		ElementOwnerStore.RegisterElementOwner(Object, CreateElement(Object));
+		ElementOwnerStore.FindOrRegisterElementOwner(Object, [Object, &CreateElement]() { return CreateElement(Object); });
 	}
 }
 
