Recreate render state on OverlayColor changes to primitive
* Some primitive types (instanced static mesh with nanite) require that render state is regenerated when the overlay color changes.

diff --git a/Engine/Plugins/Experimental/TypedElementsDataStorage/Source/TypedElementsDataStorage/Private/Compatibility/TypedElementActorViewportProcessors.cpp b/Engine/Plugins/Experimental/TypedElementsDataStorage/Source/TypedElementsDataStorage/Private/Compatibility/TypedElementActorViewportProcessors.cpp
--- a/Engine/Plugins/Experimental/TypedElementsDataStorage/Source/TypedElementsDataStorage/Private/Compatibility/TypedElementActorViewportProcessors.cpp
+++ b/Engine/Plugins/Experimental/TypedElementsDataStorage/Source/TypedElementsDataStorage/Private/Compatibility/TypedElementActorViewportProcessors.cpp
@@ -117,8 +117,8 @@ FAutoConsoleCommandWithArgsAndOutputDevice SetSelectionOverlayColorConsoleComman
 
 				for (TypedElementRowHandle Row : RowHandles)
 				{
+					DataStorage->RemoveColumn<FTypedElementViewportOverlayColorColumn>(Row);
 					DataStorage->AddColumn(Row, FTypedElementViewportOverlayColorColumn{ .OverlayColor = Color });
-					DataStorage->AddColumn<FTypedElementSyncBackToWorldTag>(Row);
 				}
 			}
 		}));
@@ -216,6 +216,7 @@ void UTypedElementActorViewportFactory::RegisterOverlayColorColumnToActor(ITyped
 					ActorInstance->ForEachComponent<UPrimitiveComponent>(bIncludeFromChildActors, [&ViewportColor](UPrimitiveComponent* PrimitiveComponent)
 					{
 						PrimitiveComponent->SetOverlayColor(ViewportColor.OverlayColor);
+						PrimitiveComponent->MarkRenderStateDirty();
 					});
 				}
 			}
@@ -237,6 +238,7 @@ void UTypedElementActorViewportFactory::RegisterOverlayColorColumnToActor(ITyped
 					ActorInstance->ForEachComponent<UPrimitiveComponent>(bIncludeFromChildActors, [](UPrimitiveComponent* PrimitiveComponent)
 					{
 						PrimitiveComponent->RemoveOverlayColor();
+						PrimitiveComponent->MarkRenderStateDirty();
 					});
 				}
 			}
