Hook into the GC to make sure objects in TEDS are properly cleaned up.

This change hookst TEDS Compatibility into the garbage collection so objects are removed from TEDS properly. The immediate effect is that several places where objects are tracked can be removed in the future and there's now a stronger guarantee that objects are still accessible in a OnRemove observer. Three new console commands were added for testing purpose:
1. TEDS.Feature.IntegrateWithGC - Enables the integration.
2. TEDS.Feature.ActorsClearedByGC - Disables places where actors were explicitly removed and lets the GC take care of it.
3. TEDS.Feature.EntitiesClearedByGC - Same as the previous but for entities.
These options are currently all off by default. Turning the actor or entity options on will automatically turn on the integration as well.

In order to let the editor know what features TEDS support the 3 main interfaces have been extended with functions to retrieve extension names. Right now this allows turning on/off the actor/entity support, but this can in the future be used for other extensions as well, in particular extension we want to keep in-house or only supply to specific customers. The console command "TEDS.Debug.ListExtensions" was added to discover the available extension.

To address an issue with TSAN, there's now a check for any undo operation happening on the game thread.

#rnx

diff --git a/Engine/Plugins/Experimental/TypedElementsDataStorage/Source/TypedElementsDataStorage/Private/Debug/TypedElementDatabaseDebugCommands.cpp b/Engine/Plugins/Experimental/TypedElementsDataStorage/Source/TypedElementsDataStorage/Private/Debug/TypedElementDatabaseDebugCommands.cpp
--- a/Engine/Plugins/Experimental/TypedElementsDataStorage/Source/TypedElementsDataStorage/Private/Debug/TypedElementDatabaseDebugCommands.cpp
+++ b/Engine/Plugins/Experimental/TypedElementsDataStorage/Source/TypedElementsDataStorage/Private/Debug/TypedElementDatabaseDebugCommands.cpp
@@ -6,6 +6,8 @@
 #include "Elements/Framework/TypedElementQueryBuilder.h"
 #include "Elements/Framework/TypedElementRegistry.h"
 #include "Elements/Interfaces/TypedElementDataStorageInterface.h"
+#include "Elements/Interfaces/TypedElementDataStorageCompatibilityInterface.h"
+#include "Elements/Interfaces/TypedElementDataStorageUiInterface.h"
 #include "HAL/IConsoleManager.h"
 #include "Misc/OutputDevice.h"
 #include "ProfilingDebugging/CpuProfilerTrace.h"
@@ -167,3 +169,40 @@ FAutoConsoleCommandWithOutputDevice PrintActorLabelsConsoleCommand(
 			Output.Log(TEXT("End of Typed Elements Data Storage actors list."));
 		}));
 
+FAutoConsoleCommandWithOutputDevice ListExtensionsConsoleCommand(
+	TEXT("TEDS.Debug.ListExtensions"),
+	TEXT("Prints a list for all available extension names."),
+	FConsoleCommandWithOutputDeviceDelegate::CreateLambda([](FOutputDevice& Output)
+		{
+			TRACE_CPUPROFILER_EVENT_SCOPE(TEDS.Debug.ListExtensions);
+
+			FString Message;
+			auto RecordExtensions = [&Message](FName Extension)
+			{
+				Message += TEXT("    ");
+				Extension.AppendString(Message);
+				Message += TEXT('\n');
+			};
+
+			UTypedElementRegistry* Registry = UTypedElementRegistry::GetInstance();
+
+			if (const ITypedElementDataStorageInterface* DataStorage = Registry->GetDataStorage())
+			{
+				Message = TEXT("Data Storage Extensions: \n");
+				DataStorage->ListExtensions(RecordExtensions);
+			}
+			if (const ITypedElementDataStorageCompatibilityInterface* DataStorageCompat = Registry->GetDataStorageCompatibility())
+			{
+				Message += TEXT("Data Storage Compatibility Extensions: \n");
+				DataStorageCompat->ListExtensions(RecordExtensions);
+			}
+			if (const ITypedElementDataStorageUiInterface* DataStorageUi = Registry->GetDataStorageUi())
+			{
+				Message += TEXT("Data Storage UI Extensions: \n");
+				DataStorageUi->ListExtensions(RecordExtensions);
+			}
+
+			Output.Log(Message);
+		}
+	));
+
diff --git a/Engine/Plugins/Experimental/TypedElementsDataStorage/Source/TypedElementsDataStorage/Private/TypedElementDatabase.cpp b/Engine/Plugins/Experimental/TypedElementsDataStorage/Source/TypedElementsDataStorage/Private/TypedElementDatabase.cpp
--- a/Engine/Plugins/Experimental/TypedElementsDataStorage/Source/TypedElementsDataStorage/Private/TypedElementDatabase.cpp
+++ b/Engine/Plugins/Experimental/TypedElementsDataStorage/Source/TypedElementsDataStorage/Private/TypedElementDatabase.cpp
@@ -928,6 +928,15 @@ void* UTypedElementDatabase::GetExternalSystemAddress(UClass* Target)
 	return nullptr;
 }
 
+bool UTypedElementDatabase::SupportsExtension(FName Extension) const
+{
+	return false;
+}
+
+void UTypedElementDatabase::ListExtensions(TFunctionRef<void(FName)> Callback) const
+{
+}
+
 void UTypedElementDatabase::PreparePhase(EQueryTickPhase Phase, float DeltaTime)
 {
 	if (ActiveEditorEntityManager)
diff --git a/Engine/Plugins/Experimental/TypedElementsDataStorage/Source/TypedElementsDataStorage/Private/TypedElementDatabase.h b/Engine/Plugins/Experimental/TypedElementsDataStorage/Source/TypedElementsDataStorage/Private/TypedElementDatabase.h
--- a/Engine/Plugins/Experimental/TypedElementsDataStorage/Source/TypedElementsDataStorage/Private/TypedElementDatabase.h
+++ b/Engine/Plugins/Experimental/TypedElementsDataStorage/Source/TypedElementsDataStorage/Private/TypedElementDatabase.h
@@ -136,6 +136,9 @@ public:
 	FTypedElementOnDataStorageUpdate& OnUpdateCompleted() override;
 	bool IsAvailable() const override;
 	void* GetExternalSystemAddress(UClass* Target) override;
+
+	bool SupportsExtension(FName Extension) const override;
+	void ListExtensions(TFunctionRef<void(FName)> Callback) const override;
 	
 	TSharedPtr<FTypedElementDatabaseEnvironment> GetEnvironment();
 	TSharedPtr<const FTypedElementDatabaseEnvironment> GetEnvironment() const;
diff --git a/Engine/Plugins/Experimental/TypedElementsDataStorage/Source/TypedElementsDataStorage/Private/TypedElementDatabaseCompatibility.cpp b/Engine/Plugins/Experimental/TypedElementsDataStorage/Source/TypedElementsDataStorage/Private/TypedElementDatabaseCompatibility.cpp
--- a/Engine/Plugins/Experimental/TypedElementsDataStorage/Source/TypedElementsDataStorage/Private/TypedElementDatabaseCompatibility.cpp
+++ b/Engine/Plugins/Experimental/TypedElementsDataStorage/Source/TypedElementsDataStorage/Private/TypedElementDatabaseCompatibility.cpp
@@ -15,15 +15,48 @@
 #include "Elements/Columns/TypedElementTypeInfoColumns.h"
 #include "Elements/Framework/TypedElementIndexHasher.h"
 #include "Elements/Framework/TypedElementQueryBuilder.h"
+#include "HAL/IConsoleManager.h"
 #include "Memento/TypedElementMementoRowTypes.h"
 #include "TypedElementDatabase.h"
 #include "TypedElementDatabaseEnvironment.h"
 #include "TypedElementDataStorageProfilingMacros.h"
 
+namespace TypedElementDataStorage
+{
+	bool bIntegratedWithGC = false;
+	FAutoConsoleVariableRef CVarIntegratedWithGC(
+		TEXT("TEDS.Feature.IntegrateWithGC"),
+		bIntegratedWithGC,
+		TEXT("Enables TEDS Compatibility integration with the garbage collection. This automatically cleans up garbage collected objects from TEDS."));
+
+	bool bActorsClearedByGC = false;
+	FAutoConsoleVariableRef CVarActorsClearedByGC(
+		TEXT("TEDS.Feature.ActorsClearedByGC"),
+		bActorsClearedByGC,
+		TEXT("Enables actors being removed through the garbage collection instead of requiring explicit removal. Enables `TEDS.Feature.IntegrateWithGC`."));
+
+	bool bEntitiesClearedByGC = false;
+	FAutoConsoleVariableRef CVarEntitiesClearedByGC(
+		TEXT("TEDS.Feature.EntitiesClearedByGC"),
+		bEntitiesClearedByGC,
+		TEXT("Enables entities being removed through the garbage collection instead of requiring explicit removal. Enables `TEDS.Feature.IntegrateWithGC`."));
+}
+
+static const FName IntegratedWithGCExtensionName(TEXT("IntegratedWithGC"));
+static const FName ActorsClearedByGCExtensionName(TEXT("ActorsClearedByGCExtension"));
+static const FName EntitiesClearedByGCExtensionName(TEXT("EntitiesClearedByGCExtension"));
+
 void UTypedElementDatabaseCompatibility::Initialize(UTypedElementDatabase* InStorage)
 {
+	using namespace TypedElementDataStorage;
+
 	checkf(InStorage, TEXT("Typed Element's Database compatibility manager is being initialized with an invalid storage target."));
 	
+	if (bActorsClearedByGC || bEntitiesClearedByGC)
+	{
+		bIntegratedWithGC = true;
+	}
+
 	Storage = InStorage;
 	Environment = InStorage->GetEnvironment();
 	Prepare();
@@ -34,22 +67,38 @@ void UTypedElementDatabaseCompatibility::Initialize(UTypedElementDatabase* InSto
 	PostEditChangePropertyDelegateHandle = FCoreUObjectDelegates::OnObjectPropertyChanged.AddUObject(this, &UTypedElementDatabaseCompatibility::OnPostEditChangeProperty);
 	ObjectModifiedDelegateHandle = FCoreUObjectDelegates::OnObjectModified.AddUObject(this, &UTypedElementDatabaseCompatibility::OnObjectModified);
 	ObjectReinstancedDelegateHandle = FCoreUObjectDelegates::OnObjectsReinstanced.AddUObject(this, &UTypedElementDatabaseCompatibility::OnObjectReinstanced);
-
-
-	PostWorldInitializationDelegateHandle = FWorldDelegates::OnPostWorldInitialization.AddUObject(this, &UTypedElementDatabaseCompatibility::OnPostWorldInitialization);
-	PreWorldFinishDestroyDelegateHandle = FWorldDelegates::OnPreWorldFinishDestroy.AddUObject(this, &UTypedElementDatabaseCompatibility::OnPreWorldFinishDestroy);
+	
+	if (bIntegratedWithGC)
+	{
+		PostGcUnreachableAnalysisHandle = FCoreUObjectDelegates::PostReachabilityAnalysis.AddUObject(this, &UTypedElementDatabaseCompatibility::OnPostGcUnreachableAnalysis);
+	}
+	if (!bActorsClearedByGC)
+	{
+		// Used to get all the worlds and register the actor create/destroy handles on them.
+		PostWorldInitializationDelegateHandle = FWorldDelegates::OnPostWorldInitialization.AddUObject(this, &UTypedElementDatabaseCompatibility::OnPostWorldInitialization);
+		PreWorldFinishDestroyDelegateHandle = FWorldDelegates::OnPreWorldFinishDestroy.AddUObject(this, &UTypedElementDatabaseCompatibility::OnPreWorldFinishDestroy);
+	}
 }
 
 void UTypedElementDatabaseCompatibility::Deinitialize()
 {
-	for (TPair<UWorld*, FDelegateHandle>& It : ActorDestroyedDelegateHandles)
+	using namespace TypedElementDataStorage;
+
+	if (!bActorsClearedByGC)
+	{
+		for (TPair<UWorld*, FDelegateHandle>& It : ActorDestroyedDelegateHandles)
+		{
+			It.Key->RemoveOnActorDestroyededHandler(It.Value);
+		}
+
+		FWorldDelegates::OnPreWorldFinishDestroy.Remove(PreWorldFinishDestroyDelegateHandle);
+		FWorldDelegates::OnPostWorldInitialization.Remove(PostWorldInitializationDelegateHandle);
+	}
+	if (bIntegratedWithGC)
 	{
-		It.Key->RemoveOnActorDestroyededHandler(It.Value);
+		FCoreUObjectDelegates::PostReachabilityAnalysis.Remove(PostGcUnreachableAnalysisHandle);
 	}
 
-	FWorldDelegates::OnPreWorldFinishDestroy.Remove(PreWorldFinishDestroyDelegateHandle);
-	FWorldDelegates::OnPostWorldInitialization.Remove(PostWorldInitializationDelegateHandle);
-	
 	FCoreUObjectDelegates::OnObjectsReinstanced.Remove(ObjectReinstancedDelegateHandle);
 	FCoreUObjectDelegates::OnObjectModified.Remove(ObjectModifiedDelegateHandle);
 	FCoreUObjectDelegates::OnObjectPropertyChanged.Remove(PostEditChangePropertyDelegateHandle);
@@ -175,6 +224,46 @@ TypedElementRowHandle UTypedElementDatabaseCompatibility::FindRowWithCompatibleO
 	return (Object && Storage && Storage->IsAvailable()) ? Storage->FindIndexedRow(GenerateIndexHash(Object)) : InvalidRowHandle;
 }
 
+bool UTypedElementDatabaseCompatibility::SupportsExtension(FName Extension) const
+{
+	using namespace TypedElementDataStorage;
+
+	if (Extension == IntegratedWithGCExtensionName)
+	{
+		return bIntegratedWithGC;
+	}
+	else if (Extension == ActorsClearedByGCExtensionName)
+	{
+		return bActorsClearedByGC;
+	}
+	else if (Extension == EntitiesClearedByGCExtensionName)
+	{
+		return bEntitiesClearedByGC;
+	}
+	else
+	{
+		return false;
+	}
+}
+
+void UTypedElementDatabaseCompatibility::ListExtensions(TFunctionRef<void(FName)> Callback) const
+{
+	using namespace TypedElementDataStorage;
+
+	if (bIntegratedWithGC)
+	{
+		Callback(IntegratedWithGCExtensionName);
+	}
+	if (bActorsClearedByGC)
+	{
+		Callback(ActorsClearedByGCExtensionName);
+	}
+	if (bEntitiesClearedByGC)
+	{
+		Callback(EntitiesClearedByGCExtensionName);
+	}
+}
+
 void UTypedElementDatabaseCompatibility::Prepare()
 {
 	CreateStandardArchetypes();
@@ -188,7 +277,7 @@ void UTypedElementDatabaseCompatibility::Reset()
 void UTypedElementDatabaseCompatibility::CreateStandardArchetypes()
 {
 	StandardActorTable = Storage->RegisterTable(TTypedElementColumnTypeList<
-			FTypedElementUObjectColumn, FTypedElementClassTypeInfoColumn,
+			FTypedElementUObjectColumn, FTypedElementUObjectIdColumn, FTypedElementClassTypeInfoColumn,
 			FTypedElementLabelColumn, FTypedElementLabelHashColumn, FTypedElementActorTag,
 			FTypedElementSyncFromWorldTag>(), 
 		FName("Editor_StandardActorTable"));
@@ -198,7 +287,7 @@ void UTypedElementDatabaseCompatibility::CreateStandardArchetypes()
 		FName("Editor_StandardActorWithTransformTable"));
 
 	StandardUObjectTable = Storage->RegisterTable(TTypedElementColumnTypeList<
-			FTypedElementUObjectColumn, FTypedElementClassTypeInfoColumn,
+			FTypedElementUObjectColumn, FTypedElementUObjectIdColumn, FTypedElementClassTypeInfoColumn,
 			FTypedElementSyncFromWorldTag>(), 
 		FName("Editor_StandardUObjectTable"));
 
@@ -224,6 +313,14 @@ void UTypedElementDatabaseCompatibility::RegisterTypeInformationQueries()
 		Select()
 			.ReadWrite<FTypedElementScriptStructTypeInfoColumn>()
 		.Compile());
+
+	if (TypedElementDataStorage::bIntegratedWithGC)
+	{
+		UObjectQuery = Storage->RegisterQuery(
+			Select()
+				.ReadWrite<FTypedElementUObjectIdColumn>()
+			.Compile());
+	}
 }
 
 bool UTypedElementDatabaseCompatibility::ShouldAddObject(const UObject* Object) const
@@ -275,7 +372,7 @@ TypedElementRowHandle UTypedElementDatabaseCompatibility::AddCompatibleObjectExp
 
 		if constexpr (bEnableTransactions)
 		{
-			if (GUndo)
+			if (IsInGameThread() && GUndo)
 			{
 				GUndo->StoreUndo(this, MakeUnique<FRegistrationCommandChange>(this, Object));
 			}
@@ -286,33 +383,46 @@ TypedElementRowHandle UTypedElementDatabaseCompatibility::AddCompatibleObjectExp
 }
 
 template<bool bEnableTransactions>
-void UTypedElementDatabaseCompatibility::RemoveCompatibleObjectExplicitTransactionable(UObject* Object)
+void UTypedElementDatabaseCompatibility::RemoveCompatibleObjectExplicitTransactionable(const UObject* Object)
 {
 	using namespace TypedElementDataStorage;
 
-	checkf(Storage, 
+	checkf(Storage,
 		TEXT("Removing compatible objects is not supported before Typed Element's Database compatibility manager has been initialized."));
 	IndexHash Hash = GenerateIndexHash(Object);
 	RowHandle Row = Storage->FindIndexedRow(Hash);
+
 	if (Storage->IsRowAvailable(Row))
 	{
-		const FTypedElementClassTypeInfoColumn* TypeInfoColumn = Storage->GetColumn<FTypedElementClassTypeInfoColumn>(Row);
-		if (Storage->HasRowBeenAssigned(Row) && 
-			ensureMsgf(TypeInfoColumn, TEXT("Missing type information for removed UObject at ptr 0x%p [%s]"), Object, *Object->GetName()))
-		{
-			OnPreObjectRemoved(Object, TypeInfoColumn->TypeInfo.Get(), Row);
+		RemoveCompatibleObjectExplicitTransactionable<bEnableTransactions>(Object, Row);
+	}
+}
 
-			if constexpr (bEnableTransactions)
+template<bool bEnableTransactions>
+void UTypedElementDatabaseCompatibility::RemoveCompatibleObjectExplicitTransactionable(
+	const UObject* Object, TypedElementDataStorage::RowHandle ObjectRow)
+{
+	using namespace TypedElementDataStorage;
+
+	checkf(Storage, 
+		TEXT("Removing compatible objects is not supported before Typed Element's Database compatibility manager has been initialized."));
+	
+	const FTypedElementClassTypeInfoColumn* TypeInfoColumn = Storage->GetColumn<FTypedElementClassTypeInfoColumn>(ObjectRow);
+	if (Storage->HasRowBeenAssigned(ObjectRow) &&
+		ensureMsgf(TypeInfoColumn, TEXT("Missing type information for removed UObject at ptr 0x%p [%s]"), Object, *Object->GetName()))
+	{
+		OnPreObjectRemoved(Object, TypeInfoColumn->TypeInfo.Get(), ObjectRow);
+
+		if constexpr (bEnableTransactions)
+		{
+			if (IsInGameThread() && GUndo)
 			{
-				if (GUndo)
-				{
-					GUndo->StoreUndo(this, MakeUnique<FDeregistrationCommandChange>(this, Object));
-				}
+				GUndo->StoreUndo(this, MakeUnique<FDeregistrationCommandChange>(this, const_cast<UObject*>(Object)));
 			}
 		}
-
-		Storage->RemoveRow(Row);
 	}
+
+	Storage->RemoveRow(ObjectRow);
 }
 
 TypedElementRowHandle UTypedElementDatabaseCompatibility::DealiasObject(const UObject* Object) const
@@ -604,6 +714,11 @@ void UTypedElementDatabaseCompatibility::TickPendingUObjectRegistration()
 			{
 				ITypedElementDataStorageInterface* Interface = Storage;
 				Interface->AddColumn(Row, FTypedElementUObjectColumn{ .Object = Object });
+				Interface->AddColumn(Row, FTypedElementUObjectIdColumn
+					{ 
+						.Id = Object->GetUniqueID(), 
+						.SerialNumber = GUObjectArray.GetSerialNumber(Object->GetUniqueID())
+					});
 				Interface->AddColumn(Row, FTypedElementClassTypeInfoColumn{ .TypeInfo = Object->GetClass() });
 				// Make sure the new row is tagged for update.
 				Interface->AddColumn<FTypedElementSyncFromWorldTag>(Row);
@@ -728,6 +843,36 @@ void UTypedElementDatabaseCompatibility::OnObjectReinstanced(const FCoreUObjectD
 	PendingTypeInformationUpdate.AddTypeInformation(ReplacedObjects);
 }
 
+void UTypedElementDatabaseCompatibility::OnPostGcUnreachableAnalysis()
+{
+	using namespace TypedElementDataStorage;
+	using namespace TypedElementQueryBuilder;
+
+	TArray<TPair<FUObjectItem*, RowHandle>> DeletedObjects;
+
+	TEDS_EVENT_SCOPE(TEXT("Post GC clean up"))
+
+	Storage->RunQuery(UObjectQuery, CreateDirectQueryCallbackBinding(
+		[&DeletedObjects](RowHandle Row, const FTypedElementUObjectIdColumn& ObjectId)
+		{
+			FUObjectItem* Description = GUObjectArray.IndexToObject(ObjectId.Id);
+			if (ensureMsgf(Description && Description->SerialNumber == ObjectId.SerialNumber, 
+				TEXT("The UObject found in TEDS no longer exists. TEDS was likely not informed in an earlier GC pass.")))
+				// Unable to provide additional information such as the UObject's name as the UObject will not be valid.
+			{
+				if (Description->HasAnyFlags(EInternalObjectFlags::Garbage | EInternalObjectFlags::Unreachable))
+				{
+					DeletedObjects.Emplace(Description, Row);
+				}
+			}
+		}));
+
+	for (TPair<FUObjectItem*, RowHandle>& Object : DeletedObjects)
+	{
+		RemoveCompatibleObjectExplicitTransactionable<false>(Cast<UObject>(Object.Key->Object), Object.Value);
+	}
+}
+
 void UTypedElementDatabaseCompatibility::OnPostWorldInitialization(UWorld* World, const UWorld::InitializationValues InitializationValues)
 {
 	FDelegateHandle Handle = World->AddOnActorDestroyedHandler(
diff --git a/Engine/Plugins/Experimental/TypedElementsDataStorage/Source/TypedElementsDataStorage/Private/TypedElementDatabaseCompatibility.h b/Engine/Plugins/Experimental/TypedElementsDataStorage/Source/TypedElementsDataStorage/Private/TypedElementDatabaseCompatibility.h
--- a/Engine/Plugins/Experimental/TypedElementsDataStorage/Source/TypedElementsDataStorage/Private/TypedElementDatabaseCompatibility.h
+++ b/Engine/Plugins/Experimental/TypedElementsDataStorage/Source/TypedElementsDataStorage/Private/TypedElementDatabaseCompatibility.h
@@ -56,6 +56,9 @@ public:
 
 	TypedElementRowHandle FindRowWithCompatibleObjectExplicit(const UObject* Object) const override;
 	TypedElementRowHandle FindRowWithCompatibleObjectExplicit(const void* Object) const override;
+
+	bool SupportsExtension(FName Extension) const override;
+	void ListExtensions(TFunctionRef<void(FName)> Callback) const override;
 	
 private:
 
@@ -103,7 +106,9 @@ private:
 	template<bool bEnableTransactions>
 	TypedElementRowHandle AddCompatibleObjectExplicitTransactionable(UObject* Object);
 	template<bool bEnableTransactions>
-	void RemoveCompatibleObjectExplicitTransactionable(UObject* Object);
+	void RemoveCompatibleObjectExplicitTransactionable(const UObject* Object);
+	template<bool bEnableTransactions>
+	void RemoveCompatibleObjectExplicitTransactionable(const UObject* Object, TypedElementDataStorage::RowHandle ObjectRow);
 	TypedElementRowHandle DealiasObject(const UObject* Object) const;
 
 	void Tick();
@@ -118,6 +123,8 @@ private:
 	void OnPreObjectRemoved(const void* Object, FTypedElementDatabaseCompatibilityObjectTypeInfo TypeInfo, TypedElementRowHandle Row) const;
 	void OnObjectReinstanced(const FCoreUObjectDelegates::FReplacementObjectMap& ReplacedObjects);
 
+	void OnPostGcUnreachableAnalysis();
+
 	void OnPostWorldInitialization(UWorld* World, const UWorld::InitializationValues InitializationValues);
 	void OnPreWorldFinishDestroy(UWorld* World);
 
@@ -177,6 +184,7 @@ private:
 			const TFunctionRef<void(TypedElementRowHandle, const AddressType&)>& SetupRowCallback);
 		void Reset();
 	};
+
 	PendingRegistration<TWeakObjectPtr<UObject>> UObjectsPendingRegistration;
 	PendingRegistration<ExternalObjectRegistration> ExternalObjectsPendingRegistration;
 	TArray<TypedElementDataStorage::RowHandle> RowScratchBuffer;
@@ -220,10 +228,12 @@ private:
 	FDelegateHandle PostWorldInitializationDelegateHandle;
 	FDelegateHandle PreWorldFinishDestroyDelegateHandle;
 	FDelegateHandle ObjectReinstancedDelegateHandle;
+	FDelegateHandle PostGcUnreachableAnalysisHandle;
 	
 	TSharedPtr<FTypedElementDatabaseEnvironment> Environment;
 	TypedElementDataStorage::QueryHandle ClassTypeInfoQuery;
 	TypedElementDataStorage::QueryHandle ScriptStructTypeInfoQuery;
+	TypedElementDataStorage::QueryHandle UObjectQuery;
 };
 
 SIZE_T GetTypeHash(const UTypedElementDatabaseCompatibility::FSyncTagInfo& Column);
diff --git a/Engine/Plugins/Experimental/TypedElementsDataStorage/Source/TypedElementsDataStorage/Private/TypedElementDatabaseUI.cpp b/Engine/Plugins/Experimental/TypedElementsDataStorage/Source/TypedElementsDataStorage/Private/TypedElementDatabaseUI.cpp
--- a/Engine/Plugins/Experimental/TypedElementsDataStorage/Source/TypedElementsDataStorage/Private/TypedElementDatabaseUI.cpp
+++ b/Engine/Plugins/Experimental/TypedElementsDataStorage/Source/TypedElementsDataStorage/Private/TypedElementDatabaseUI.cpp
@@ -448,6 +448,15 @@ void UTypedElementDatabaseUi::ListWidgetPurposes(const WidgetPurposeCallback& Ca
 	}
 }
 
+bool UTypedElementDatabaseUi::SupportsExtension(FName Extension) const
+{
+	return false;
+}
+
+void UTypedElementDatabaseUi::ListExtensions(TFunctionRef<void(FName)> Callback) const
+{
+}
+
 void UTypedElementDatabaseUi::CreateStandardArchetypes()
 {
 	WidgetTable = Storage->RegisterTable(MakeArrayView(
diff --git a/Engine/Plugins/Experimental/TypedElementsDataStorage/Source/TypedElementsDataStorage/Private/TypedElementDatabaseUI.h b/Engine/Plugins/Experimental/TypedElementsDataStorage/Source/TypedElementsDataStorage/Private/TypedElementDatabaseUI.h
--- a/Engine/Plugins/Experimental/TypedElementsDataStorage/Source/TypedElementsDataStorage/Private/TypedElementDatabaseUI.h
+++ b/Engine/Plugins/Experimental/TypedElementsDataStorage/Source/TypedElementsDataStorage/Private/TypedElementDatabaseUI.h
@@ -52,6 +52,9 @@ public:
 
 	void ListWidgetPurposes(const WidgetPurposeCallback& Callback) const override;
 
+	bool SupportsExtension(FName Extension) const override;
+	void ListExtensions(TFunctionRef<void(FName)> Callback) const override;
+
 private:
 	struct FWidgetFactory
 	{
diff --git a/Engine/Source/Editor/UnrealEd/Private/LevelEditorViewport.cpp b/Engine/Source/Editor/UnrealEd/Private/LevelEditorViewport.cpp
--- a/Engine/Source/Editor/UnrealEd/Private/LevelEditorViewport.cpp
+++ b/Engine/Source/Editor/UnrealEd/Private/LevelEditorViewport.cpp
@@ -1665,24 +1665,27 @@ void FLevelEditorViewportClient::DestroyDropPreviewElements()
 		return;
 	}
 
-	// TODO: This code to remove the object from TEDS should not be necessary, beause the element
+	// TODO: This code to remove the object from TEDS should not be necessary, because the element
 	// deletion code further below should include TEDS deregistration. However, the code path for
 	// deleting preview actors in UUnrealEdEngine::DeleteActors skips explicit handle deregistration,
 	// and although it does still happen in the immediately triggered garbage cleanup, that feels
 	// potentially brittle.
-	DropPreviewElements->ForEachElement<ITypedElementObjectInterface>([this](const TTypedElement<ITypedElementObjectInterface>& InElement)
+	static FName ActorsClearedByGCExtensionName(TEXT("ActorsClearedByGCExtension"));
+	UTypedElementRegistry* TypedElementRegistry = UTypedElementRegistry::GetInstance();
+	if (ITypedElementDataStorageCompatibilityInterface* TedsCompat = TypedElementRegistry->GetMutableDataStorageCompatibility();
+		TedsCompat && !TedsCompat->SupportsExtension(ActorsClearedByGCExtensionName))
 	{
-		UObject* PreviewObject = InElement.GetObject();
-		if (PreviewObject && PreviewObject != GetWorld()->GetDefaultBrush())
-		{
-			UTypedElementRegistry* TypedElementRegistry = UTypedElementRegistry::GetInstance();
-			if (ITypedElementDataStorageCompatibilityInterface* DataStorageCompatibilityInterface = TypedElementRegistry->GetMutableDataStorageCompatibility())
+		DropPreviewElements->ForEachElement<ITypedElementObjectInterface>(
+			[this, TedsCompat](const TTypedElement<ITypedElementObjectInterface>& InElement)
 			{
-				DataStorageCompatibilityInterface->RemoveCompatibleObject(PreviewObject);
-			}
-		}
-		return true; // true means continue
-	});
+				UObject* PreviewObject = InElement.GetObject();
+				if (PreviewObject && PreviewObject != GetWorld()->GetDefaultBrush())
+				{
+					TedsCompat->RemoveCompatibleObject(PreviewObject);
+				}
+				return true; // true means continue
+			});
+	}
 
 	// Used for special casing BSP backwards compatibility: the builder brush is used as a preview object, and
 	// we don't want to delete it.
diff --git a/Engine/Source/Runtime/Engine/Private/Elements/Framework/EngineElementsLibrary.cpp b/Engine/Source/Runtime/Engine/Private/Elements/Framework/EngineElementsLibrary.cpp
--- a/Engine/Source/Runtime/Engine/Private/Elements/Framework/EngineElementsLibrary.cpp
+++ b/Engine/Source/Runtime/Engine/Private/Elements/Framework/EngineElementsLibrary.cpp
@@ -389,9 +389,11 @@ void UEngineElementsLibrary::RegisterActorElement(const AActor* InActor)
 
 void UEngineElementsLibrary::UnregisterActorElement(const AActor* InActor)
 {
+	static FName ActorsClearedByGCExtensionName = TEXT("ActorsClearedByGCExtension");
+
 	UTypedElementRegistry* Registry = UTypedElementRegistry::GetInstance();
 	ITypedElementDataStorageCompatibilityInterface* Storage = Registry->GetMutableDataStorageCompatibility();
-	if (Storage)
+	if (Storage && !Storage->SupportsExtension(ActorsClearedByGCExtensionName))
 	{
 		Storage->RemoveCompatibleObject(const_cast<AActor*>(InActor));
 	}
diff --git a/Engine/Source/Runtime/TypedElementFramework/Public/Elements/Columns/TypedElementCompatibilityColumns.h b/Engine/Source/Runtime/TypedElementFramework/Public/Elements/Columns/TypedElementCompatibilityColumns.h
--- a/Engine/Source/Runtime/TypedElementFramework/Public/Elements/Columns/TypedElementCompatibilityColumns.h
+++ b/Engine/Source/Runtime/TypedElementFramework/Public/Elements/Columns/TypedElementCompatibilityColumns.h
@@ -20,9 +20,26 @@ struct FTypedElementUObjectColumn final : public FTypedElementDataStorageColumn
 {
 	GENERATED_BODY()
 
+	// If the UObject is accessed during an OnRemove event that is triggered by the garbage collection, the UObject will
+	// already be marked as unreachable and regular functions to retrieve the UObject will return an nullptr. In these cases
+	// use the unreachable versions such as Get(/*bEvenIfPendingKill*/ true) or GetEvenIfUnreachable().
 	TWeakObjectPtr<UObject> Object;
 };
 
+/**
+ * Column containing information to uniquely identify the UObject, e.g. for use by the garbage collection.
+ */
+USTRUCT(meta = (DisplayName = "UObject ID"))
+struct FTypedElementUObjectIdColumn final : public FTypedElementDataStorageColumn
+{
+	GENERATED_BODY()
+
+	UPROPERTY()
+	uint32 Id;
+	UPROPERTY()
+	int32 SerialNumber;
+};
+
 /**
  * Column containing a non-owning reference to an arbitrary object. It's strongly recommended
  * to also add a FTypedElementScriptStructTypeInfoColumn to make sure the type can be safely
diff --git a/Engine/Source/Runtime/TypedElementFramework/Public/Elements/Interfaces/TypedElementDataStorageCompatibilityInterface.h b/Engine/Source/Runtime/TypedElementFramework/Public/Elements/Interfaces/TypedElementDataStorageCompatibilityInterface.h
--- a/Engine/Source/Runtime/TypedElementFramework/Public/Elements/Interfaces/TypedElementDataStorageCompatibilityInterface.h
+++ b/Engine/Source/Runtime/TypedElementFramework/Public/Elements/Interfaces/TypedElementDataStorageCompatibilityInterface.h
@@ -99,6 +99,15 @@ public:
 	virtual TypedElementRowHandle FindRowWithCompatibleObjectExplicit(const UObject* Object) const = 0;
 	/** Finds a previously stored FStruct. If not found an invalid row handle will be returned. */
 	virtual TypedElementRowHandle FindRowWithCompatibleObjectExplicit(const void* Object) const = 0;
+
+	/**
+	 * @section Miscellaneous functions
+	 */
+
+	/** Check if a custom extension is supported. This can be used to check for in-development features, custom extensions, etc. */
+	virtual bool SupportsExtension(FName Extension) const = 0;
+	/** Provides a list of all extensions that are enabled. */
+	virtual void ListExtensions(TFunctionRef<void(FName)> Callback) const = 0;
 };
 
 template<typename Type> Type* GetRawPointer(const TWeakObjectPtr<Type> Object)	{ return Object.Get(); }
diff --git a/Engine/Source/Runtime/TypedElementFramework/Public/Elements/Interfaces/TypedElementDataStorageInterface.h b/Engine/Source/Runtime/TypedElementFramework/Public/Elements/Interfaces/TypedElementDataStorageInterface.h
--- a/Engine/Source/Runtime/TypedElementFramework/Public/Elements/Interfaces/TypedElementDataStorageInterface.h
+++ b/Engine/Source/Runtime/TypedElementFramework/Public/Elements/Interfaces/TypedElementDataStorageInterface.h
@@ -289,7 +289,7 @@ public:
 	virtual void RemoveIndex(TypedElementDataStorage::IndexHash Index) = 0;
 
 	/**
-	 * @section Misc
+	 * @section Miscellaneous
 	 */
 	
 	/**
@@ -311,10 +311,14 @@ public:
 
 	/** Returns a pointer to the registered external system if found, otherwise null. */
 	virtual void* GetExternalSystemAddress(UClass* Target) = 0;
-
 	/** Returns a pointer to the registered external system if found, otherwise null. */
 	template<typename SystemType>
 	SystemType* GetExternalSystem();
+
+	/** Check if a custom extension is supported. This can be used to check for in-development features, custom extensions, etc. */
+	virtual bool SupportsExtension(FName Extension) const = 0;
+	/** Provides a list of all extensions that are enabled. */
+	virtual void ListExtensions(TFunctionRef<void(FName)> Callback) const = 0;
 };
 
 
diff --git a/Engine/Source/Runtime/TypedElementFramework/Public/Elements/Interfaces/TypedElementDataStorageUiInterface.h b/Engine/Source/Runtime/TypedElementFramework/Public/Elements/Interfaces/TypedElementDataStorageUiInterface.h
--- a/Engine/Source/Runtime/TypedElementFramework/Public/Elements/Interfaces/TypedElementDataStorageUiInterface.h
+++ b/Engine/Source/Runtime/TypedElementFramework/Public/Elements/Interfaces/TypedElementDataStorageUiInterface.h
@@ -257,6 +257,11 @@ public:
 
 	/** Calls the provided callback for all known registered widget purposes. */
 	virtual void ListWidgetPurposes(const WidgetPurposeCallback& Callback) const = 0;
+
+	/** Check if a custom extension is supported. This can be used to check for in-development features, custom extensions, etc. */
+	virtual bool SupportsExtension(FName Extension) const = 0;
+	/** Provides a list of all extensions that are enabled. */
+	virtual void ListExtensions(TFunctionRef<void(FName)> Callback) const = 0;
 };
 
 
