Addressed the issue of CVars in TEDS Compat being used too early.

When CVars are set externally through something like the command line or config files they're not available yet when a module is initialized. This was causing the new CVars for TEDS Compat to use the GC to not work properly.

#rnx

diff --git a/Engine/Plugins/Experimental/TypedElementsDataStorage/Source/TypedElementsDataStorage/Private/TypedElementDatabaseCompatibility.cpp b/Engine/Plugins/Experimental/TypedElementsDataStorage/Source/TypedElementsDataStorage/Private/TypedElementDatabaseCompatibility.cpp
--- a/Engine/Plugins/Experimental/TypedElementsDataStorage/Source/TypedElementsDataStorage/Private/TypedElementDatabaseCompatibility.cpp
+++ b/Engine/Plugins/Experimental/TypedElementsDataStorage/Source/TypedElementsDataStorage/Private/TypedElementDatabaseCompatibility.cpp
@@ -23,12 +23,6 @@
 
 namespace TypedElementDataStorage
 {
-	bool bIntegratedWithGC = false;
-	FAutoConsoleVariableRef CVarIntegratedWithGC(
-		TEXT("TEDS.Feature.IntegrateWithGC"),
-		bIntegratedWithGC,
-		TEXT("Enables TEDS Compatibility integration with the garbage collection. This automatically cleans up garbage collected objects from TEDS."));
-
 	bool bActorsClearedByGC = false;
 	FAutoConsoleVariableRef CVarActorsClearedByGC(
 		TEXT("TEDS.Feature.ActorsClearedByGC"),
@@ -42,7 +36,6 @@ namespace TypedElementDataStorage
 		TEXT("Enables entities being removed through the garbage collection instead of requiring explicit removal. Enables `TEDS.Feature.IntegrateWithGC`."));
 }
 
-static const FName IntegratedWithGCExtensionName(TEXT("IntegratedWithGC"));
 static const FName ActorsClearedByGCExtensionName(TEXT("ActorsClearedByGCExtension"));
 static const FName EntitiesClearedByGCExtensionName(TEXT("EntitiesClearedByGCExtension"));
 
@@ -52,11 +45,6 @@ void UTypedElementDatabaseCompatibility::Initialize(UTypedElementDatabase* InSto
 
 	checkf(InStorage, TEXT("Typed Element's Database compatibility manager is being initialized with an invalid storage target."));
 	
-	if (bActorsClearedByGC || bEntitiesClearedByGC)
-	{
-		bIntegratedWithGC = true;
-	}
-
 	Storage = InStorage;
 	Environment = InStorage->GetEnvironment();
 	Prepare();
@@ -68,37 +56,25 @@ void UTypedElementDatabaseCompatibility::Initialize(UTypedElementDatabase* InSto
 	ObjectModifiedDelegateHandle = FCoreUObjectDelegates::OnObjectModified.AddUObject(this, &UTypedElementDatabaseCompatibility::OnObjectModified);
 	ObjectReinstancedDelegateHandle = FCoreUObjectDelegates::OnObjectsReinstanced.AddUObject(this, &UTypedElementDatabaseCompatibility::OnObjectReinstanced);
 	
-	if (bIntegratedWithGC)
-	{
-		PostGcUnreachableAnalysisHandle = FCoreUObjectDelegates::PostReachabilityAnalysis.AddUObject(this, &UTypedElementDatabaseCompatibility::OnPostGcUnreachableAnalysis);
-	}
-	if (!bActorsClearedByGC)
-	{
-		// Used to get all the worlds and register the actor create/destroy handles on them.
-		PostWorldInitializationDelegateHandle = FWorldDelegates::OnPostWorldInitialization.AddUObject(this, &UTypedElementDatabaseCompatibility::OnPostWorldInitialization);
-		PreWorldFinishDestroyDelegateHandle = FWorldDelegates::OnPreWorldFinishDestroy.AddUObject(this, &UTypedElementDatabaseCompatibility::OnPreWorldFinishDestroy);
-	}
+	PostGcUnreachableAnalysisHandle = FCoreUObjectDelegates::PostReachabilityAnalysis.AddUObject(this, &UTypedElementDatabaseCompatibility::OnPostGcUnreachableAnalysis);
+	// Used to get all the worlds and register the actor create/destroy handles on them.
+	PostWorldInitializationDelegateHandle = FWorldDelegates::OnPostWorldInitialization.AddUObject(this, &UTypedElementDatabaseCompatibility::OnPostWorldInitialization);
+	PreWorldFinishDestroyDelegateHandle = FWorldDelegates::OnPreWorldFinishDestroy.AddUObject(this, &UTypedElementDatabaseCompatibility::OnPreWorldFinishDestroy);
 }
 
 void UTypedElementDatabaseCompatibility::Deinitialize()
 {
 	using namespace TypedElementDataStorage;
 
-	if (!bActorsClearedByGC)
+	for (TPair<UWorld*, FDelegateHandle>& It : ActorDestroyedDelegateHandles)
 	{
-		for (TPair<UWorld*, FDelegateHandle>& It : ActorDestroyedDelegateHandles)
-		{
-			It.Key->RemoveOnActorDestroyededHandler(It.Value);
-		}
-
-		FWorldDelegates::OnPreWorldFinishDestroy.Remove(PreWorldFinishDestroyDelegateHandle);
-		FWorldDelegates::OnPostWorldInitialization.Remove(PostWorldInitializationDelegateHandle);
-	}
-	if (bIntegratedWithGC)
-	{
-		FCoreUObjectDelegates::PostReachabilityAnalysis.Remove(PostGcUnreachableAnalysisHandle);
+		It.Key->RemoveOnActorDestroyededHandler(It.Value);
 	}
 
+	FWorldDelegates::OnPreWorldFinishDestroy.Remove(PreWorldFinishDestroyDelegateHandle);
+	FWorldDelegates::OnPostWorldInitialization.Remove(PostWorldInitializationDelegateHandle);
+	FCoreUObjectDelegates::PostReachabilityAnalysis.Remove(PostGcUnreachableAnalysisHandle);
+
 	FCoreUObjectDelegates::OnObjectsReinstanced.Remove(ObjectReinstancedDelegateHandle);
 	FCoreUObjectDelegates::OnObjectModified.Remove(ObjectModifiedDelegateHandle);
 	FCoreUObjectDelegates::OnObjectPropertyChanged.Remove(PostEditChangePropertyDelegateHandle);
@@ -228,11 +204,7 @@ bool UTypedElementDatabaseCompatibility::SupportsExtension(FName Extension) cons
 {
 	using namespace TypedElementDataStorage;
 
-	if (Extension == IntegratedWithGCExtensionName)
-	{
-		return bIntegratedWithGC;
-	}
-	else if (Extension == ActorsClearedByGCExtensionName)
+	if (Extension == ActorsClearedByGCExtensionName)
 	{
 		return bActorsClearedByGC;
 	}
@@ -250,10 +222,6 @@ void UTypedElementDatabaseCompatibility::ListExtensions(TFunctionRef<void(FName)
 {
 	using namespace TypedElementDataStorage;
 
-	if (bIntegratedWithGC)
-	{
-		Callback(IntegratedWithGCExtensionName);
-	}
 	if (bActorsClearedByGC)
 	{
 		Callback(ActorsClearedByGCExtensionName);
@@ -314,13 +282,10 @@ void UTypedElementDatabaseCompatibility::RegisterTypeInformationQueries()
 			.ReadWrite<FTypedElementScriptStructTypeInfoColumn>()
 		.Compile());
 
-	if (TypedElementDataStorage::bIntegratedWithGC)
-	{
-		UObjectQuery = Storage->RegisterQuery(
-			Select()
-				.ReadWrite<FTypedElementUObjectIdColumn>()
-			.Compile());
-	}
+	UObjectQuery = Storage->RegisterQuery(
+		Select()
+			.ReadWrite<FTypedElementUObjectIdColumn>()
+		.Compile());
 }
 
 bool UTypedElementDatabaseCompatibility::ShouldAddObject(const UObject* Object) const
@@ -848,49 +813,60 @@ void UTypedElementDatabaseCompatibility::OnPostGcUnreachableAnalysis()
 	using namespace TypedElementDataStorage;
 	using namespace TypedElementQueryBuilder;
 
-	TArray<TPair<FUObjectItem*, RowHandle>> DeletedObjects;
+	if (bActorsClearedByGC || bEntitiesClearedByGC)
+	{
+		TArray<TPair<FUObjectItem*, RowHandle>> DeletedObjects;
 
-	TEDS_EVENT_SCOPE(TEXT("Post GC clean up"))
+		TEDS_EVENT_SCOPE(TEXT("Post GC clean up"))
 
-	Storage->RunQuery(UObjectQuery, CreateDirectQueryCallbackBinding(
-		[&DeletedObjects](RowHandle Row, const FTypedElementUObjectIdColumn& ObjectId)
-		{
-			FUObjectItem* Description = GUObjectArray.IndexToObject(ObjectId.Id);
-			if (ensureMsgf(Description && Description->SerialNumber == ObjectId.SerialNumber, 
-				TEXT("The UObject found in TEDS no longer exists. TEDS was likely not informed in an earlier GC pass.")))
-				// Unable to provide additional information such as the UObject's name as the UObject will not be valid.
+		Storage->RunQuery(UObjectQuery, CreateDirectQueryCallbackBinding(
+			[&DeletedObjects](RowHandle Row, const FTypedElementUObjectIdColumn& ObjectId)
 			{
-				if (Description->HasAnyFlags(EInternalObjectFlags::Garbage | EInternalObjectFlags::Unreachable))
+				FUObjectItem* Description = GUObjectArray.IndexToObject(ObjectId.Id);
+				if (ensureMsgf(Description && Description->SerialNumber == ObjectId.SerialNumber, 
+					TEXT("The UObject found in TEDS no longer exists. TEDS was likely not informed in an earlier GC pass.")))
+					// Unable to provide additional information such as the UObject's name as the UObject will not be valid.
 				{
-					DeletedObjects.Emplace(Description, Row);
+					if (Description->HasAnyFlags(EInternalObjectFlags::Garbage | EInternalObjectFlags::Unreachable))
+					{
+						DeletedObjects.Emplace(Description, Row);
+					}
 				}
-			}
-		}));
+			}));
 
-	for (TPair<FUObjectItem*, RowHandle>& Object : DeletedObjects)
-	{
-		RemoveCompatibleObjectExplicitTransactionable<false>(Cast<UObject>(Object.Key->Object), Object.Value);
+		for (TPair<FUObjectItem*, RowHandle>& Object : DeletedObjects)
+		{
+			RemoveCompatibleObjectExplicitTransactionable<false>(Cast<UObject>(Object.Key->Object), Object.Value);
+		}
 	}
 }
 
 void UTypedElementDatabaseCompatibility::OnPostWorldInitialization(UWorld* World, const UWorld::InitializationValues InitializationValues)
 {
-	FDelegateHandle Handle = World->AddOnActorDestroyedHandler(
-		FOnActorDestroyed::FDelegate::CreateUObject(this, &UTypedElementDatabaseCompatibility::OnActorDestroyed));
-	ActorDestroyedDelegateHandles.Add(World, Handle);
+	if (!TypedElementDataStorage::bActorsClearedByGC)
+	{
+		FDelegateHandle Handle = World->AddOnActorDestroyedHandler(
+			FOnActorDestroyed::FDelegate::CreateUObject(this, &UTypedElementDatabaseCompatibility::OnActorDestroyed));
+		ActorDestroyedDelegateHandles.Add(World, Handle);
+	}
 }
 
 void UTypedElementDatabaseCompatibility::OnPreWorldFinishDestroy(UWorld* World)
 {
-	FDelegateHandle Handle;
-	if (ActorDestroyedDelegateHandles.RemoveAndCopyValue(World, Handle))
+	if (!TypedElementDataStorage::bActorsClearedByGC)
 	{
-		World->RemoveOnActorDestroyededHandler(Handle);
+		FDelegateHandle Handle;
+		if (ActorDestroyedDelegateHandles.RemoveAndCopyValue(World, Handle))
+		{
+			World->RemoveOnActorDestroyededHandler(Handle);
+		}
 	}
 }
 
 void UTypedElementDatabaseCompatibility::OnActorDestroyed(AActor* Actor)
 {
+	checkf(!TypedElementDataStorage::bActorsClearedByGC, 
+		TEXT("If actors are cleared from TEDS Compatibility through garbage collection, there shouldn't be any explicitly actor destruction handlers."));
 	RemoveCompatibleObjectExplicit(Actor);
 }
 
