Add FTypedElementClassDefaultObjectTag and use it to exclude CDOs from the query "Add world column to UObject".

diff --git a/Engine/Plugins/Experimental/TypedElementsDataStorage/Source/TypedElementsDataStorage/Private/Compatibility/TypedElementUObjectWorldQueries.cpp b/Engine/Plugins/Experimental/TypedElementsDataStorage/Source/TypedElementsDataStorage/Private/Compatibility/TypedElementUObjectWorldQueries.cpp
--- a/Engine/Plugins/Experimental/TypedElementsDataStorage/Source/TypedElementsDataStorage/Private/Compatibility/TypedElementUObjectWorldQueries.cpp
+++ b/Engine/Plugins/Experimental/TypedElementsDataStorage/Source/TypedElementsDataStorage/Private/Compatibility/TypedElementUObjectWorldQueries.cpp
@@ -38,7 +38,7 @@ void UTypedElementUObjectWorldFactory::RegisterAddWorldColumn(ITypedElementDataS
 		)
 		.Where()
 			.All<FTypedElementSyncFromWorldTag>()
-			.None<FTypedElementWorldColumn>()
+			.None<FTypedElementWorldColumn, FTypedElementClassDefaultObjectTag>()
 		.Compile()
 	);
 }
diff --git a/Engine/Plugins/Experimental/TypedElementsDataStorage/Source/TypedElementsDataStorage/Private/TypedElementDatabaseCompatibility.cpp b/Engine/Plugins/Experimental/TypedElementsDataStorage/Source/TypedElementsDataStorage/Private/TypedElementDatabaseCompatibility.cpp
--- a/Engine/Plugins/Experimental/TypedElementsDataStorage/Source/TypedElementsDataStorage/Private/TypedElementDatabaseCompatibility.cpp
+++ b/Engine/Plugins/Experimental/TypedElementsDataStorage/Source/TypedElementsDataStorage/Private/TypedElementDatabaseCompatibility.cpp
@@ -685,6 +685,10 @@ void UTypedElementDatabaseCompatibility::TickPendingUObjectRegistration()
 						.SerialNumber = GUObjectArray.GetSerialNumber(Object->GetUniqueID())
 					});
 				Interface->AddColumn(Row, FTypedElementClassTypeInfoColumn{ .TypeInfo = Object->GetClass() });
+				if (Object->HasAnyFlags(RF_ClassDefaultObject))
+				{
+					Interface->AddColumn<FTypedElementClassDefaultObjectTag>(Row);
+				}
 				// Make sure the new row is tagged for update.
 				Interface->AddColumn<FTypedElementSyncFromWorldTag>(Row);
 				OnObjectAdded(Object.Get(), Object->GetClass(), Row);
diff --git a/Engine/Source/Runtime/TypedElementFramework/Public/Elements/Columns/TypedElementCompatibilityColumns.h b/Engine/Source/Runtime/TypedElementFramework/Public/Elements/Columns/TypedElementCompatibilityColumns.h
--- a/Engine/Source/Runtime/TypedElementFramework/Public/Elements/Columns/TypedElementCompatibilityColumns.h
+++ b/Engine/Source/Runtime/TypedElementFramework/Public/Elements/Columns/TypedElementCompatibilityColumns.h
@@ -53,6 +53,16 @@ struct FTypedElementExternalObjectColumn final : public FTypedElementDataStorage
 	void* Object;
 };
 
+/**
+ * Tag to identify a row with a Class Default Object (CDO).
+ * If there's a FTypedElementUObjectColumn, the stored object will be a CDO if this tag is present.
+ */
+USTRUCT(meta = (DisplayName = "Class Default Object"))
+struct FTypedElementClassDefaultObjectTag final : public FTypedElementDataStorageTag
+{
+	GENERATED_BODY()
+};
+
 /**
  * Tag to identify a row with an actor. If there's a FTypedElementUObjectColumn, the stored object will be
  * an actor if this tag is present.
