Added AddRow variants to TEDS that provide an opportunity to initialize columns.

When adding a row to TEDS it's instantly created based on the provided table. This also causes observers to be immediately called which is either more expensive or causes subtle bugs. `BatchAddRow` already handles this through a callback and with this CR, `AddRow` now also has an option for that.

The following additional changes were made:
- All the table and row functions have had their use of `TypedElement*` updated to `TypedElementDataStorage::`.
- The implementation of `BatchAddRow` for reserved row handles no longer uses an intermediate array but instead uses the provided array through casting. This is a minor performance gain.

#rnx

diff --git a/Engine/Plugins/Experimental/TypedElementsDataStorage/Source/TypedElementsDataStorage/Private/TypedElementDatabase.cpp b/Engine/Plugins/Experimental/TypedElementsDataStorage/Source/TypedElementsDataStorage/Private/TypedElementDatabase.cpp
--- a/Engine/Plugins/Experimental/TypedElementsDataStorage/Source/TypedElementsDataStorage/Private/TypedElementDatabase.cpp
+++ b/Engine/Plugins/Experimental/TypedElementsDataStorage/Source/TypedElementsDataStorage/Private/TypedElementDatabase.cpp
@@ -263,11 +263,13 @@ TSharedPtr<const FMassEntityManager> UTypedElementDatabase::GetActiveEditorEntit
 	return ActiveEditorEntityManager;
 }
 
-TypedElementTableHandle UTypedElementDatabase::RegisterTable(TConstArrayView<const UScriptStruct*> ColumnList, const FName Name)
+TypedElementDataStorage::TableHandle UTypedElementDatabase::RegisterTable(TConstArrayView<const UScriptStruct*> ColumnList, const FName Name)
 {
+	using namespace TypedElementDataStorage;
+
 	if (ActiveEditorEntityManager && !TableNameLookup.Contains(Name))
 	{
-		TypedElementTableHandle Result = Tables.Num();
+		TableHandle Result = Tables.Num();
 		FMassArchetypeCreationParams ArchetypeCreationParams;
 		ArchetypeCreationParams.DebugName = Name;
 		ArchetypeCreationParams.ChunkMemorySize = GetTableChunkSize(Name);
@@ -278,15 +280,17 @@ TypedElementTableHandle UTypedElementDatabase::RegisterTable(TConstArrayView<con
 		}
 		return Result;
 	}
-	return TypedElementInvalidTableHandle;
+	return InvalidTableHandle;
 }
 
-TypedElementTableHandle UTypedElementDatabase::RegisterTable(TypedElementTableHandle SourceTable, 
+TypedElementDataStorage::TableHandle UTypedElementDatabase::RegisterTable(TypedElementDataStorage::TableHandle SourceTable,
 	TConstArrayView<const UScriptStruct*> ColumnList, const FName Name)
 {
+	using namespace TypedElementDataStorage;
+
 	if (ActiveEditorEntityManager && SourceTable < Tables.Num() && !TableNameLookup.Contains(Name))
 	{
-		TypedElementTableHandle Result = Tables.Num();
+		TypedElementDataStorage::TableHandle Result = Tables.Num();
 		FMassArchetypeCreationParams ArchetypeCreationParams;
 		ArchetypeCreationParams.DebugName = Name;
 		ArchetypeCreationParams.ChunkMemorySize = GetTableChunkSize(Name);
@@ -297,29 +301,56 @@ TypedElementTableHandle UTypedElementDatabase::RegisterTable(TypedElementTableHa
 		}
 		return Result;
 	}
-	return TypedElementInvalidTableHandle;
+	return InvalidTableHandle;
 }
 
-TypedElementTableHandle UTypedElementDatabase::FindTable(const FName Name)
+TypedElementDataStorage::TableHandle UTypedElementDatabase::FindTable(const FName Name)
 {
-	TypedElementTableHandle* TableHandle = TableNameLookup.Find(Name);
-	return TableHandle ? *TableHandle : TypedElementInvalidTableHandle;
+	using namespace TypedElementDataStorage;
+
+	TableHandle* TableHandle = TableNameLookup.Find(Name);
+	return TableHandle ? *TableHandle : InvalidTableHandle;
 }
 
-TypedElementRowHandle UTypedElementDatabase::ReserveRow()
+TypedElementDataStorage::RowHandle UTypedElementDatabase::ReserveRow()
 {
-	return ActiveEditorEntityManager ? ActiveEditorEntityManager->ReserveEntity().AsNumber(): TypedElementInvalidRowHandle;
+	return ActiveEditorEntityManager 
+		? ActiveEditorEntityManager->ReserveEntity().AsNumber()
+		: TypedElementDataStorage::InvalidRowHandle;
 }
 
-TypedElementRowHandle UTypedElementDatabase::AddRow(TypedElementTableHandle Table)
+TypedElementDataStorage::RowHandle UTypedElementDatabase::AddRow(TypedElementDataStorage::TableHandle Table)
 {
 	checkf(Table < Tables.Num(), TEXT("Attempting to add a row to a non-existing table."));
-	return ActiveEditorEntityManager ? 
-		ActiveEditorEntityManager->CreateEntity(Tables[Table]).AsNumber() :
-		TypedElementInvalidRowHandle;
+	return ActiveEditorEntityManager 
+		? ActiveEditorEntityManager->CreateEntity(Tables[Table]).AsNumber() 
+		: TypedElementDataStorage::InvalidRowHandle;
 }
 
-bool UTypedElementDatabase::AddRow(TypedElementRowHandle ReservedRow, TypedElementTableHandle Table)
+TypedElementDataStorage::RowHandle UTypedElementDatabase::AddRow(TypedElementDataStorage::TableHandle Table,
+	TypedElementDataStorage::RowCreationCallbackRef OnCreated)
+{
+	using namespace TypedElementDataStorage;
+
+	OnCreated.CheckCallable();
+	if (ActiveEditorEntityManager)
+	{
+		checkf(Table < Tables.Num(), TEXT("Attempting to a row to a non-existing table."));
+
+		TArray<FMassEntityHandle> Entity;
+		Entity.Reserve(1);
+		TSharedRef<FMassEntityManager::FEntityCreationContext> Context =
+			ActiveEditorEntityManager->BatchCreateEntities(Tables[Table], 1, Entity);
+
+		checkf(!Entity.IsEmpty(), TEXT("Add row tried to create a new row but none were provided by the backend."));
+		RowHandle Result = Entity[0].AsNumber();
+		OnCreated(Entity[0].AsNumber());
+		return Result;
+	}
+	return InvalidRowHandle;
+}
+
+bool UTypedElementDatabase::AddRow(TypedElementDataStorage::RowHandle ReservedRow, TypedElementDataStorage::TableHandle Table)
 {
 	checkf(!IsRowAssigned(ReservedRow), TEXT("Attempting to assign a table to row that already has a table assigned."));
 	checkf(Table < Tables.Num(), TEXT("Attempting to add a row to a non-existing table."));
@@ -334,12 +365,31 @@ bool UTypedElementDatabase::AddRow(TypedElementRowHandle ReservedRow, TypedEleme
 	}
 }
 
-bool UTypedElementDatabase::BatchAddRow(TypedElementTableHandle Table, int32 Count, TypedElementDataStorage::RowCreationCallbackRef OnCreated)
+bool UTypedElementDatabase::AddRow(TypedElementDataStorage::RowHandle ReservedRow, TypedElementDataStorage::TableHandle Table,
+	TypedElementDataStorage::RowCreationCallbackRef OnCreated)
 {
 	OnCreated.CheckCallable();
-	checkf(Table < Tables.Num(), TEXT("Attempting to add multiple rows to a non-existing table."));
 	if (ActiveEditorEntityManager)
 	{
+		checkf(Table < Tables.Num(), TEXT("Attempting to add a row to a non-existing table."));
+		
+		TSharedRef<FMassEntityManager::FEntityCreationContext> Context =
+			ActiveEditorEntityManager->BatchCreateReservedEntities(Tables[Table], { FMassEntityHandle::FromNumber(ReservedRow) });
+
+		OnCreated(ReservedRow);
+		return true;
+	}
+	return false;
+}
+
+bool UTypedElementDatabase::BatchAddRow(
+	TypedElementDataStorage::TableHandle Table, int32 Count, TypedElementDataStorage::RowCreationCallbackRef OnCreated)
+{
+	OnCreated.CheckCallable();
+	if (ActiveEditorEntityManager)
+	{
+		checkf(Table < Tables.Num(), TEXT("Attempting to add multiple rows to a non-existing table."));
+	
 		TArray<FMassEntityHandle> Entities;
 		Entities.Reserve(Count);
 		TSharedRef<FMassEntityManager::FEntityCreationContext> Context = 
@@ -355,28 +405,27 @@ bool UTypedElementDatabase::BatchAddRow(TypedElementTableHandle Table, int32 Cou
 	return false;
 }
 
-bool UTypedElementDatabase::BatchAddRow(TypedElementTableHandle Table, TConstArrayView<TypedElementRowHandle> ReservedHandles,
-	TypedElementDataStorage::RowCreationCallbackRef OnCreated)
+bool UTypedElementDatabase::BatchAddRow(TypedElementDataStorage::TableHandle Table, 
+	TConstArrayView<TypedElementDataStorage::RowHandle> ReservedHandles, TypedElementDataStorage::RowCreationCallbackRef OnCreated)
 {
 	OnCreated.CheckCallable();
-	checkf(Table < Tables.Num(), TEXT("Attempting to add multiple rows to a non-existing table."));
 	if (ActiveEditorEntityManager)
 	{
-		TArray<FMassEntityHandle> Entities;
-		Entities.Reserve(ReservedHandles.Num());
-		FMassEntityHandle* CurrentEntityHandle = Entities.GetData();
-		for (TypedElementRowHandle RowHandle : ReservedHandles)
-		{
-			checkf(!IsRowAssigned(RowHandle), TEXT("Attempting to assign a table to row that already has a table assigned."));
-			Entities.Add(FMassEntityHandle::FromNumber(RowHandle));
-		}
+		checkf(Table < Tables.Num(), TEXT("Attempting to add multiple rows to a non-existing table."));
+	
+		// Depend on the fact that a row handle is an alias for an entity within the Mass powered backend. This
+		// avoids the need for copying to a temporary array;
+		static_assert(sizeof(TypedElementRowHandle) == sizeof(FMassEntityHandle), 
+			"BatchAddRow in TEDS requires the row handle and the Mass entity handle to be the same size.");
 		
+		TConstArrayView<FMassEntityHandle> Entities(
+			reinterpret_cast<const FMassEntityHandle*>(ReservedHandles.GetData()), ReservedHandles.Num());
 		TSharedRef<FMassEntityManager::FEntityCreationContext> Context =
 			ActiveEditorEntityManager->BatchCreateReservedEntities(Tables[Table], Entities);
 
-		for (FMassEntityHandle Entity : Entities)
+		for (TypedElementRowHandle Entity : ReservedHandles)
 		{
-			OnCreated(Entity.AsNumber());
+			OnCreated(Entity);
 		}
 
 		return true;
@@ -385,7 +434,7 @@ bool UTypedElementDatabase::BatchAddRow(TypedElementTableHandle Table, TConstArr
 }
 
 
-void UTypedElementDatabase::RemoveRow(TypedElementRowHandle Row)
+void UTypedElementDatabase::RemoveRow(TypedElementDataStorage::RowHandle Row)
 {
 	FMassEntityHandle Entity = FMassEntityHandle::FromNumber(Row);
 	if (ActiveEditorEntityManager && ActiveEditorEntityManager->IsEntityValid(Entity))
@@ -403,12 +452,12 @@ void UTypedElementDatabase::RemoveRow(TypedElementRowHandle Row)
 	}
 }
 
-bool UTypedElementDatabase::IsRowAvailable(TypedElementRowHandle Row) const
+bool UTypedElementDatabase::IsRowAvailable(TypedElementDataStorage::RowHandle Row) const
 {
 	return ActiveEditorEntityManager ? FTypedElementDatabaseCommandBuffer::Execute_IsRowAvailable(*ActiveEditorEntityManager, Row) : false;
 }
 
-bool UTypedElementDatabase::IsRowAssigned(TypedElementRowHandle Row) const
+bool UTypedElementDatabase::IsRowAssigned(TypedElementDataStorage::RowHandle Row) const
 {
 	return ActiveEditorEntityManager ? FTypedElementDatabaseCommandBuffer::Execute_IsRowAssigned(*ActiveEditorEntityManager, Row) : false;
 }
diff --git a/Engine/Plugins/Experimental/TypedElementsDataStorage/Source/TypedElementsDataStorage/Private/TypedElementDatabase.h b/Engine/Plugins/Experimental/TypedElementsDataStorage/Source/TypedElementsDataStorage/Private/TypedElementDatabase.h
--- a/Engine/Plugins/Experimental/TypedElementsDataStorage/Source/TypedElementsDataStorage/Private/TypedElementDatabase.h
+++ b/Engine/Plugins/Experimental/TypedElementsDataStorage/Source/TypedElementsDataStorage/Private/TypedElementDatabase.h
@@ -81,20 +81,25 @@ public:
 	TSharedPtr<FMassEntityManager> GetActiveMutableEditorEntityManager();
 	TSharedPtr<const FMassEntityManager> GetActiveEditorEntityManager() const;
 
-	TypedElementTableHandle RegisterTable(TConstArrayView<const UScriptStruct*> ColumnList, const FName Name) override;
-	TypedElementTableHandle RegisterTable(TypedElementTableHandle SourceTable, TConstArrayView<const UScriptStruct*> ColumnList, 
-		const FName Name) override;
-	TypedElementTableHandle FindTable(const FName Name) override;
-
-	TypedElementRowHandle ReserveRow() override;
-	TypedElementRowHandle AddRow(TypedElementTableHandle Table) override;
-	bool AddRow(TypedElementRowHandle ReservedRow, TypedElementTableHandle Table) override;
-	bool BatchAddRow(TypedElementTableHandle Table, int32 Count, TypedElementDataStorage::RowCreationCallbackRef OnCreated) override;
-	bool BatchAddRow(TypedElementTableHandle Table, TConstArrayView<TypedElementRowHandle> ReservedHandles,
+	TypedElementDataStorage::TableHandle RegisterTable(TConstArrayView<const UScriptStruct*> ColumnList, const FName Name) override;
+	TypedElementDataStorage::TableHandle RegisterTable(
+		TypedElementDataStorage::TableHandle SourceTable, TConstArrayView<const UScriptStruct*> ColumnList, const FName Name) override;
+	TypedElementDataStorage::TableHandle FindTable(const FName Name) override;
+
+	TypedElementDataStorage::RowHandle ReserveRow() override;
+	TypedElementDataStorage::RowHandle AddRow(TypedElementDataStorage::TableHandle Table, 
+		TypedElementDataStorage::RowCreationCallbackRef OnCreated) override;
+	TypedElementDataStorage::RowHandle AddRow(TypedElementDataStorage::TableHandle Table) override;
+	bool AddRow(TypedElementDataStorage::RowHandle ReservedRow, TypedElementDataStorage::TableHandle Table) override;
+	bool AddRow(TypedElementDataStorage::RowHandle ReservedRow, TypedElementDataStorage::TableHandle Table,
+		TypedElementDataStorage::RowCreationCallbackRef OnCreated) override;
+	bool BatchAddRow(TypedElementDataStorage::TableHandle Table, int32 Count, 
+		TypedElementDataStorage::RowCreationCallbackRef OnCreated) override;
+	bool BatchAddRow(TypedElementDataStorage::TableHandle Table, TConstArrayView<TypedElementDataStorage::RowHandle> ReservedHandles,
 		TypedElementDataStorage::RowCreationCallbackRef OnCreated) override;
-	void RemoveRow(TypedElementRowHandle Row) override;
-	bool IsRowAvailable(TypedElementRowHandle Row) const override;
-	bool IsRowAssigned(TypedElementRowHandle Row) const override;
+	void RemoveRow(TypedElementDataStorage::RowHandle Row) override;
+	bool IsRowAvailable(TypedElementDataStorage::RowHandle Row) const override;
+	bool IsRowAssigned(TypedElementDataStorage::RowHandle Row) const override;
 
 	void AddColumn(TypedElementRowHandle Row, const UScriptStruct* ColumnType) override;
 	void AddColumnData(TypedElementRowHandle Row, const UScriptStruct* ColumnType,
@@ -171,7 +176,7 @@ private:
 	static const FName TickGroupName_SyncWidget;
 	
 	TArray<FMassArchetypeHandle> Tables;
-	TMap<FName, TypedElementTableHandle> TableNameLookup;
+	TMap<FName, TypedElementDataStorage::TableHandle> TableNameLookup;
 
 	// Ordered array of factories by the return value of GetOrder()
 	TArray<FFactoryTypePair> Factories;
diff --git a/Engine/Source/Runtime/TypedElementFramework/Public/Elements/Interfaces/TypedElementDataStorageInterface.h b/Engine/Source/Runtime/TypedElementFramework/Public/Elements/Interfaces/TypedElementDataStorageInterface.h
--- a/Engine/Source/Runtime/TypedElementFramework/Public/Elements/Interfaces/TypedElementDataStorageInterface.h
+++ b/Engine/Source/Runtime/TypedElementFramework/Public/Elements/Interfaces/TypedElementDataStorageInterface.h
@@ -111,32 +111,44 @@ public:
 	 * Reserves a row to be assigned to a table at a later point. If the row is no longer needed before it's been assigned
 	 * to a table, it should still be released with RemoveRow.
 	 */
-	virtual TypedElementRowHandle ReserveRow() = 0;
+	virtual TypedElementDataStorage::RowHandle ReserveRow() = 0;
 	/** Adds a new row to the provided table. */
-	virtual TypedElementRowHandle AddRow(TypedElementTableHandle Table) = 0;
-	/** Adds a new row to the provided table using a previously reserved row.. */
-	virtual bool AddRow(TypedElementRowHandle ReservedRow, TypedElementTableHandle Table) = 0;
-	
+	virtual TypedElementDataStorage::RowHandle AddRow(TypedElementDataStorage::TableHandle Table) = 0;
+	/**
+	 * Adds a new row to the provided table. Callers are expected to use the callback to
+	 * initialize the row if needed.
+	 */
+	virtual TypedElementDataStorage::RowHandle AddRow(TypedElementDataStorage::TableHandle Table,
+		TypedElementDataStorage::RowCreationCallbackRef OnCreated) = 0;
+	/** Adds a new row to the provided table using a previously reserved row. */
+	virtual bool AddRow(TypedElementDataStorage::RowHandle ReservedRow, TypedElementDataStorage::TableHandle Table) = 0;
+	/**
+	 * Adds a new row to the provided table using a previously reserved row. Callers are expected to use the callback to
+	 * initialize the row if needed.
+	 */
+	virtual bool AddRow(TypedElementDataStorage::RowHandle ReservedRow, TypedElementDataStorage::TableHandle Table,
+		TypedElementDataStorage::RowCreationCallbackRef OnCreated) = 0;
+
 	/**
 	 * Add multiple rows at once. For each new row the OnCreated callback is called. Callers are expected to use the callback to
 	 * initialize the row if needed.
 	 */
-	virtual bool BatchAddRow(TypedElementTableHandle Table, int32 Count, TypedElementDataStorage::RowCreationCallbackRef OnCreated) = 0;
+	virtual bool BatchAddRow(TypedElementDataStorage::TableHandle Table, int32 Count, TypedElementDataStorage::RowCreationCallbackRef OnCreated) = 0;
 	/**
 	 * Add multiple rows at once. For each new row the OnCreated callback is called. Callers are expected to use the callback to
 	 * initialize the row if needed. This version uses a set of previously reserved rows. Any row that can't be used will be 
 	 * released.
 	 */
-	virtual bool BatchAddRow(TypedElementTableHandle Table, TConstArrayView<TypedElementRowHandle> ReservedHandles,
+	virtual bool BatchAddRow(TypedElementDataStorage::TableHandle Table, TConstArrayView<TypedElementDataStorage::RowHandle> ReservedHandles,
 		TypedElementDataStorage::RowCreationCallbackRef OnCreated) = 0;
 
 	/** Removes a previously reserved or added row. If the row handle is invalid or already removed, nothing happens */
-	virtual void RemoveRow(TypedElementRowHandle Row) = 0;
+	virtual void RemoveRow(TypedElementDataStorage::RowHandle Row) = 0;
 
 	/** Checks whether or not a row is in use. This is true even if the row has only been reserved. */
-	virtual bool IsRowAvailable(TypedElementRowHandle Row) const = 0;
+	virtual bool IsRowAvailable(TypedElementDataStorage::RowHandle Row) const = 0;
 	/** Checks whether or not a row has been reserved but not yet assigned to a table. */
-	virtual bool IsRowAssigned(TypedElementRowHandle Row) const = 0;
+	virtual bool IsRowAssigned(TypedElementDataStorage::RowHandle Row) const = 0;
 
 	
 	/**
