TEDS: Add placeholder columns for asset data and add console commands to populate placeholder asset rows and remove asset rows

diff --git a/Engine/Plugins/Experimental/TypedElementsDataStorage/Source/TypedElementsDataStorage/Private/AssetDebugCommands.cpp b/Engine/Plugins/Experimental/TypedElementsDataStorage/Source/TypedElementsDataStorage/Private/AssetDebugCommands.cpp
new file mode 100644
--- /dev/null
+++ b/Engine/Plugins/Experimental/TypedElementsDataStorage/Source/TypedElementsDataStorage/Private/AssetDebugCommands.cpp
@@ -0,0 +1,136 @@
+// Copyright Epic Games, Inc. All Rights Reserved.
+
+#include "Engine/StaticMesh.h"
+#include "Materials/Material.h"
+#include "Engine/Level.h"
+#include "Engine/Blueprint.h"
+#include "Engine/Texture.h"
+
+#include "Elements/Columns/TypedElementLabelColumns.h"
+#include "Elements/Columns/TypedElementTypeInfoColumns.h"
+#include "Elements/Framework/TypedElementQueryBuilder.h"
+#include "Elements/Framework/TypedElementRegistry.h"
+#include "HAL/IConsoleManager.h"
+#include "Math/UnrealMathUtility.h"
+#include "TypedElementAssetColumns.h"
+
+namespace UE::TEDSAssetDataDebugCommands::Local
+{
+	// A small set of classes to randomly pick from
+	static TArray<UClass*> AssetClasses{ UStaticMesh::StaticClass(), UMaterial::StaticClass(),
+		ULevel::StaticClass(), UBlueprint::StaticClass(), UTexture::StaticClass() };
+
+	// Populate an asset row with random information
+	void PopulateRowWithRandomInfo(TypedElementDataStorage::RowHandle Row, ITypedElementDataStorageInterface* DataStorage)
+	{
+		// Don't modify any rows that aren't our placeholder assets
+		if (!DataStorage->HasColumns<FAssetTag>(Row))
+		{
+			return;
+		}
+
+		// Pick a random asset class from our list
+		UClass* AssetClass = AssetClasses[FMath::RandRange(0, AssetClasses.Num() - 1)];
+
+		// Add a label to the row
+		if (FTypedElementLabelColumn* LabelColumn = DataStorage->GetColumn<FTypedElementLabelColumn>(Row))
+		{
+			// We can have duplicate names but it doesn't really matter
+			LabelColumn->Label = AssetClass->GetFName().ToString();
+			LabelColumn->Label.Append(TEXT("_Placeholder_"));
+			LabelColumn->Label.Append(FString::FromInt(FMath::RandRange(0, 1000)));
+
+			if (FVersePathColumn* ClassTypeInfoColumn = DataStorage->GetColumn<FVersePathColumn>(Row))
+			{
+				FString TestVerseModule = TEXT("/UnrealEngine.com/Temporary/TEDS/");
+				TestVerseModule.Append(LabelColumn->Label);
+			
+				UE::Core::FVersePath::TryMake(ClassTypeInfoColumn->VersePath, TestVerseModule);
+			}
+		}
+
+		if (FTypedElementClassTypeInfoColumn* ClassTypeInfoColumn = DataStorage->GetColumn<FTypedElementClassTypeInfoColumn>(Row))
+		{
+			ClassTypeInfoColumn->TypeInfo = AssetClass;
+		}
+		
+		if (FDiskSizeColumn* ClassTypeInfoColumn = DataStorage->GetColumn<FDiskSizeColumn>(Row))
+		{
+			ClassTypeInfoColumn->DiskSize = FMath::RandRange(1024, 32768);
+		}
+
+		// Randomly make this a public or a private asset
+		DataStorage->AddColumn(Row, FMath::RandRange(0, 1) == 0 ? FPrivateAssetTag::StaticStruct() : FPublicAssetTag::StaticStruct());
+	}
+	
+	FAutoConsoleCommand CreateDebugAssetRowsCommand(
+	TEXT("TEDS.Debug.CreateDebugAssetRows"),
+	TEXT("Create random asset rows. Args: (TEDS.Debug.CreateDebugAssetRows NumRows) Default 10 rows"),
+	FConsoleCommandWithArgsDelegate::CreateLambda(
+		[](const TArray<FString>& Args)
+		{
+			TRACE_CPUPROFILER_EVENT_SCOPE(TEDS.Debug.CreateDebugAssetRows);
+
+			using namespace TypedElementDataStorage;
+			using namespace TypedElementQueryBuilder;
+
+			ITypedElementDataStorageInterface* DataStorage = UTypedElementRegistry::GetInstance()->GetMutableDataStorage();
+
+			if (!DataStorage)
+			{
+				return;
+			}
+
+			int RowCount = 10; // Create 10 rows by default
+
+			if (Args.Num() == 1)
+			{
+				RowCount = FCString::Atoi(*Args[0]);
+			}
+
+			TableHandle Table = DataStorage->FindTable(FName("Editor_PlaceholderAssetTable"));
+
+			DataStorage->BatchAddRow(Table, RowCount, [DataStorage](RowHandle Row)
+			{
+				PopulateRowWithRandomInfo(Row, DataStorage);
+			});
+			
+		}
+	));
+
+	FAutoConsoleCommand RemoveAssetRowsCommand(
+	TEXT("TEDS.Debug.RemoveAssetRows"),
+	TEXT("Remove All Asset Rows"),
+	FConsoleCommandWithArgsDelegate::CreateLambda(
+		[](const TArray<FString>& Args)
+		{
+			using namespace TypedElementDataStorage;
+			using namespace TypedElementQueryBuilder;
+
+			ITypedElementDataStorageInterface* DataStorage = UTypedElementRegistry::GetInstance()->GetMutableDataStorage();
+
+			if (!DataStorage)
+			{
+				return;
+			}
+
+			// Select all rows with FAssetTag
+			static QueryHandle AssetQueryHandle = DataStorage->RegisterQuery(Select().Where().All<FAssetTag>().Compile());
+
+			TArray<RowHandle> Rows;
+			
+			DataStorage->RunQuery(AssetQueryHandle,
+				CreateDirectQueryCallbackBinding([&Rows](const ITypedElementDataStorageInterface::IDirectQueryContext& Context)
+				{
+					Rows.Append(Context.GetRowHandles());
+				}));
+
+			for(const RowHandle Row : Rows)
+			{
+				DataStorage->RemoveRow(Row);
+			}
+		}
+	));
+
+}
+
diff --git a/Engine/Plugins/Experimental/TypedElementsDataStorage/Source/TypedElementsDataStorage/Private/AssetProcessors.cpp b/Engine/Plugins/Experimental/TypedElementsDataStorage/Source/TypedElementsDataStorage/Private/AssetProcessors.cpp
new file mode 100644
--- /dev/null
+++ b/Engine/Plugins/Experimental/TypedElementsDataStorage/Source/TypedElementsDataStorage/Private/AssetProcessors.cpp
@@ -0,0 +1,16 @@
+// Copyright Epic Games, Inc. All Rights Reserved.
+
+#include "AssetProcessors.h"
+
+#include "Elements/Columns/TypedElementLabelColumns.h"
+#include "Elements/Columns/TypedElementTypeInfoColumns.h"
+#include "TypedElementAssetColumns.h"
+
+void UTypedElementAssetFactory::RegisterTables(ITypedElementDataStorageInterface& DataStorage) 
+{
+	DataStorage.RegisterTable(
+		TTypedElementColumnTypeList<
+			FTypedElementLabelColumn, FTypedElementClassTypeInfoColumn,
+			FAssetTag, FDiskSizeColumn, FVersePathColumn>(),
+		FName("Editor_PlaceholderAssetTable"));
+}
diff --git a/Engine/Plugins/Experimental/TypedElementsDataStorage/Source/TypedElementsDataStorage/Private/AssetProcessors.h b/Engine/Plugins/Experimental/TypedElementsDataStorage/Source/TypedElementsDataStorage/Private/AssetProcessors.h
new file mode 100644
--- /dev/null
+++ b/Engine/Plugins/Experimental/TypedElementsDataStorage/Source/TypedElementsDataStorage/Private/AssetProcessors.h
@@ -0,0 +1,22 @@
+// Copyright Epic Games, Inc. All Rights Reserved.
+
+#pragma once
+
+#include "Elements/Interfaces/TypedElementDataStorageFactory.h"
+#include "Elements/Interfaces/TypedElementQueryStorageInterfaces.h"
+#include "UObject/ObjectMacros.h"
+
+#include "AssetProcessors.generated.h"
+
+UCLASS()
+class UTypedElementAssetFactory : public UTypedElementDataStorageFactory
+{
+	GENERATED_BODY()
+
+public:
+	~UTypedElementAssetFactory() override = default;
+
+	void RegisterTables(ITypedElementDataStorageInterface& DataStorage) override;
+
+private:
+};
diff --git a/Engine/Plugins/Experimental/TypedElementsDataStorage/Source/TypedElementsDataStorage/Public/TypedElementAssetColumns.h b/Engine/Plugins/Experimental/TypedElementsDataStorage/Source/TypedElementsDataStorage/Public/TypedElementAssetColumns.h
new file mode 100644
--- /dev/null
+++ b/Engine/Plugins/Experimental/TypedElementsDataStorage/Source/TypedElementsDataStorage/Public/TypedElementAssetColumns.h
@@ -0,0 +1,49 @@
+// Copyright Epic Games, Inc. All Rights Reserved.
+
+#pragma once
+
+#include "Elements/Interfaces/TypedElementDataStorageInterface.h"
+#include "Containers/VersePath.h"
+
+#include "TypedElementAssetColumns.generated.h"
+
+// Tag to identify assets
+USTRUCT(meta = (DisplayName = "Asset"))
+struct FAssetTag final : public FTypedElementDataStorageTag
+{
+	GENERATED_BODY()
+};
+
+// Tag to identify assets with private visibility
+USTRUCT(meta = (DisplayName = "Private Asset"))
+struct FPrivateAssetTag final : public FTypedElementDataStorageTag
+{
+	GENERATED_BODY()
+};
+
+// Tag to identify assets with public visibility
+USTRUCT(meta = (DisplayName = "Public Asset"))
+struct FPublicAssetTag final : public FTypedElementDataStorageTag
+{
+	GENERATED_BODY()
+};
+
+// Column to store the disk size of an asset
+USTRUCT(meta = (DisplayName = "Disk Size"))
+struct FDiskSizeColumn final : public FTypedElementDataStorageColumn
+{
+	GENERATED_BODY()
+
+	// Size on disk (in bytes)
+	UPROPERTY()
+	int64 DiskSize;
+};
+
+// Column to store the verse path of an asset
+USTRUCT(meta = (DisplayName = "Verse Path"))
+struct FVersePathColumn final : public FTypedElementDataStorageColumn
+{
+	GENERATED_BODY()
+	
+	UE::Core::FVersePath VersePath;
+};
\ No newline at end of file
