The index table in TEDS has been made thread-safe.

It's now safe to call the index table in TEDS from multiple threads. A mutex is used to guarantee exclusive access. This includes when reading from the index table, except when a phase is being processed. While a phase is being processed, no write to the index table can happen (no context supports this) so the index table is locked but any reads can happen from any thread.
In addition to the thread-safety the following changes were made:
- It's now possible to batch reserve rows.
- It's now possible to batch register indices.

diff --git a/Engine/Plugins/Experimental/TypedElementsDataStorage/Source/TypedElementsDataStorage/Private/GlobalLock.cpp b/Engine/Plugins/Experimental/TypedElementsDataStorage/Source/TypedElementsDataStorage/Private/GlobalLock.cpp
new file mode 100644
--- /dev/null
+++ b/Engine/Plugins/Experimental/TypedElementsDataStorage/Source/TypedElementsDataStorage/Private/GlobalLock.cpp
@@ -0,0 +1,139 @@
+// Copyright Epic Games, Inc. All Rights Reserved.
+
+#include "GlobalLock.h"
+
+namespace UE::EditorDataStorage
+{
+	//
+	// FGlobalLock
+	//
+
+
+	FRWLock FGlobalLock::Lock;
+	thread_local EGlobalLockStatus FGlobalLock::LockStatus = EGlobalLockStatus::Unlocked;
+	thread_local uint64 FGlobalLock::LockCount = 0;
+	std::atomic<EGlobalLockStatus> FGlobalLock::InternalLockStatus = EGlobalLockStatus::Unlocked;
+
+	void FGlobalLock::SharedLock(EGlobalLockScope Scope)
+	{
+		if (Scope == EGlobalLockScope::Public)
+		{
+			// The requirements for a shared lock are also satisfied if this thread has an exclusive lock.
+			if (LockStatus == EGlobalLockStatus::Unlocked)
+			{
+				Lock.ReadLock();
+				LockStatus = EGlobalLockStatus::SharedLocked;
+			}
+			LockCount++;
+		}
+		else
+		{
+			checkf(InternalLockStatus == EGlobalLockStatus::SharedLocked || LockStatus != EGlobalLockStatus::Unlocked,
+				TEXT("Requesting an global internal shared lock in TEDS while the internal or shared lock hasn't been acquired."));
+		}
+	}
+
+	void FGlobalLock::ExclusiveLock(EGlobalLockScope Scope)
+	{
+		if (Scope == EGlobalLockScope::Public)
+		{
+			if (LockStatus != EGlobalLockStatus::ExclusiveLocked)
+			{
+				checkf(LockStatus != EGlobalLockStatus::SharedLocked,
+					TEXT("Attempting to get a global TEDS exclusive lock on a thread that's already shared locked."));
+				Lock.WriteLock();
+				LockStatus = EGlobalLockStatus::ExclusiveLocked;
+			}
+			LockCount++;
+		}
+		else
+		{
+			checkf(false, TEXT("Internal exclusive locks for TEDS can't be safely requested as internal locks are shared only."));
+		}
+	}
+
+	void FGlobalLock::Unlock(EGlobalLockScope Scope)
+	{
+		if (Scope == EGlobalLockScope::Public)
+		{
+			switch (LockStatus)
+			{
+			default:
+				// fall through.
+			case EGlobalLockStatus::Unlocked:
+				checkf(false, TEXT("Attempting to unlock the global TEDS lock that wasn't locked."));
+				break;
+			case EGlobalLockStatus::SharedLocked:
+				checkf(LockCount > 0, TEXT("Attempting to unlock the global TEDS lock while its lock count is already zero."));
+				if (--LockCount == 0)
+				{
+					LockStatus = EGlobalLockStatus::Unlocked;
+					Lock.ReadUnlock();
+				}
+				break;
+			case EGlobalLockStatus::ExclusiveLocked:
+				checkf(LockCount > 0, TEXT("Attempting to unlock the global TEDS lock while its lock count is already zero."));
+				if (--LockCount == 0)
+				{
+					LockStatus = EGlobalLockStatus::Unlocked;
+					Lock.WriteUnlock();
+				}
+				break;
+			}
+		}
+	}
+
+	EGlobalLockStatus FGlobalLock::GetLockStatus(EGlobalLockScope Scope)
+	{
+		return Scope == EGlobalLockScope::Public ? LockStatus : InternalLockStatus.load();
+	}
+
+	void FGlobalLock::InternalSharedLock()
+	{
+		checkf(InternalLockStatus == 
+			EGlobalLockStatus::Unlocked, TEXT("Attempting to acquire a global internal TEDS lock while there is already an internal lock."));
+		// Get a public shared lock on behalf of all internal locks.
+		SharedLock(EGlobalLockScope::Public);
+		InternalLockStatus = EGlobalLockStatus::SharedLocked;
+	}
+	
+	void FGlobalLock::InternalSharedUnlock()
+	{
+		checkf(InternalLockStatus ==
+			EGlobalLockStatus::SharedLocked, TEXT("Attempting to release a global internal TEDS lock while there is not internal lock."));
+		InternalLockStatus = EGlobalLockStatus::Unlocked;
+		Unlock(EGlobalLockScope::Public);
+	}
+
+
+	//
+	// FScopedSharedLock
+	//
+
+	FScopedSharedLock::FScopedSharedLock(EGlobalLockScope InScope)
+		: Scope(InScope)
+	{
+		FGlobalLock::SharedLock(Scope);
+	}
+
+	FScopedSharedLock::~FScopedSharedLock()
+	{
+		FGlobalLock::Unlock(Scope);
+	}
+
+
+	//
+	// FScopedExclusiveLock
+	//
+
+	FScopedExclusiveLock::FScopedExclusiveLock(EGlobalLockScope InScope)
+		:Scope(InScope)
+	{
+		FGlobalLock::ExclusiveLock(Scope);
+	}
+
+	FScopedExclusiveLock::~FScopedExclusiveLock()
+	{
+		FGlobalLock::Unlock(Scope);
+	}
+} // UE::EditorDataStorage
\ No newline at end of file
diff --git a/Engine/Plugins/Experimental/TypedElementsDataStorage/Source/TypedElementsDataStorage/Private/GlobalLock.h b/Engine/Plugins/Experimental/TypedElementsDataStorage/Source/TypedElementsDataStorage/Private/GlobalLock.h
new file mode 100644
--- /dev/null
+++ b/Engine/Plugins/Experimental/TypedElementsDataStorage/Source/TypedElementsDataStorage/Private/GlobalLock.h
@@ -0,0 +1,87 @@
+// Copyright Epic Games, Inc. All Rights Reserved.
+
+#pragma once
+
+#include <atomic>
+#include "HAL/CriticalSection.h"
+#include "HAL/Platform.h"
+
+class UTypedElementDatabase;
+
+namespace UE::EditorDataStorage
+{
+	enum class EGlobalLockStatus : uint8
+	{
+		Unlocked,
+		SharedLocked,
+		ExclusiveLocked
+	};
+
+	enum class EGlobalLockScope : uint8
+	{
+		/** The lock request is coming from a public facing call. */
+		Public,
+		/** 
+		 * The lock is coming from an internal call. These typically have a wider range of
+		 * calling threads, but only support shared access. Exclusive locks will check for
+		 * validity.
+		 */
+		Internal,
+	};
+
+	/**
+	 * Sets a lock for all of TEDS systems. This is used whenever access to any data within TEDS is needed in a thread-safe manner. Shared 
+	 * locks will allow multiple thread to safely read data, while exclusive locks provide safe read/write access.
+	 * 
+	 * The Global lock has an internal lock so shared locks can be freely handed out to requests with an internal scope. Exclusive 
+	 * locks with the same scope will assert. This is meant to only be called by TEDS Core to lock the global lock for the duration of
+	 * a processing phase. This is the only safe time to call it as all operations within a phase will be processed on all threads before
+	 * a phase ends. If this is used for global calls it's possible that such a call is still processing when TEDS Core releases the lock
+	 * and another exclusive lock is acquired. Therefore global locks fall back to full lock. Use global locks except from within adaptors
+	 * that run in Mass processors.
+	 * 
+	 * The global lock uses internal reference counting to allow recursive calls from the same thread. Further state management also helps
+	 * track incorrect requests for an exclusive lock on a thread that already has a shared lock. An assert triggers in this case to prevent
+	 * deadlocks.
+	 */
+	class FGlobalLock
+	{
+		friend class ::UTypedElementDatabase;
+	public:
+		static void SharedLock(EGlobalLockScope Scope);
+		static void ExclusiveLock(EGlobalLockScope Scope);
+		static void Unlock(EGlobalLockScope Scope);
+		
+		static EGlobalLockStatus GetLockStatus(EGlobalLockScope Scope);
+
+	private:
+		static void InternalSharedLock();
+		static void InternalSharedUnlock();
+
+		static FRWLock Lock;
+		/** Keep track of the lock status because FRWLock is non-recursive and with a global lock it's easy to get recursive calls. */
+		static thread_local EGlobalLockStatus LockStatus;
+		static thread_local uint64 LockCount;
+		static std::atomic<EGlobalLockStatus> InternalLockStatus;
+	};
+
+	class FScopedSharedLock
+	{
+	public:
+		explicit FScopedSharedLock(EGlobalLockScope InScope);
+		~FScopedSharedLock();
+
+	private:
+		EGlobalLockScope Scope;
+	};
+
+	class FScopedExclusiveLock
+	{
+	public:
+		explicit FScopedExclusiveLock(EGlobalLockScope InScope);
+		~FScopedExclusiveLock();
+
+	private:
+		EGlobalLockScope Scope;
+	};
+} // UE::EditorDataStorage
diff --git a/Engine/Plugins/Experimental/TypedElementsDataStorage/Source/TypedElementsDataStorage/Private/Processors/TypedElementProcessorAdaptors.cpp b/Engine/Plugins/Experimental/TypedElementsDataStorage/Source/TypedElementsDataStorage/Private/Processors/TypedElementProcessorAdaptors.cpp
--- a/Engine/Plugins/Experimental/TypedElementsDataStorage/Source/TypedElementsDataStorage/Private/Processors/TypedElementProcessorAdaptors.cpp
+++ b/Engine/Plugins/Experimental/TypedElementsDataStorage/Source/TypedElementsDataStorage/Private/Processors/TypedElementProcessorAdaptors.cpp
@@ -539,7 +539,12 @@ struct FMassQueryContextImplementation final : FMassWithEnvironmentContextCommon
 	
 	TypedElementDataStorage::RowHandle FindIndexedRow(TypedElementDataStorage::IndexHash Index) const
 	{
-		return Environment.GetIndexTable().FindIndexedRow(Index);
+		using namespace UE::EditorDataStorage;
+
+		EGlobalLockScope Scope = FGlobalLock::GetLockStatus(EGlobalLockScope::Internal) == EGlobalLockStatus::Unlocked
+			? EGlobalLockScope::Public // There's no internal lock so use a public lock instead.
+			: EGlobalLockScope::Internal; // There's an internal lock set so use that.
+		return Environment.GetIndexTable().FindIndexedRow(Scope, Index);
 	}
 
 	TypedElementDataStorage::FQueryResult RunQuery(TypedElementQueryHandle Query)
diff --git a/Engine/Plugins/Experimental/TypedElementsDataStorage/Source/TypedElementsDataStorage/Private/TypedElementDatabase.cpp b/Engine/Plugins/Experimental/TypedElementsDataStorage/Source/TypedElementsDataStorage/Private/TypedElementDatabase.cpp
--- a/Engine/Plugins/Experimental/TypedElementsDataStorage/Source/TypedElementsDataStorage/Private/TypedElementDatabase.cpp
+++ b/Engine/Plugins/Experimental/TypedElementsDataStorage/Source/TypedElementsDataStorage/Private/TypedElementDatabase.cpp
@@ -7,6 +7,7 @@
 #include "Elements/Interfaces/TypedElementDataStorageFactory.h"
 #include "Elements/Framework/TypedElementRegistry.h"
 #include "Engine/World.h"
+#include "GlobalLock.h"
 #include "MassCommonTypes.h"
 #include "MassEntityEditorSubsystem.h"
 #include "MassEntityTypes.h"
@@ -319,6 +320,38 @@ TypedElementDataStorage::RowHandle UTypedElementDatabase::ReserveRow()
 		: TypedElementDataStorage::InvalidRowHandle;
 }
 
+void UTypedElementDatabase::BatchReserveRows(int32 Count, TFunctionRef<void(TypedElementDataStorage::RowHandle)> ReservationCallback)
+{
+	using namespace TypedElementDataStorage;
+
+	if (ActiveEditorEntityManager)
+	{
+		TArrayView<FMassEntityHandle> ReservedEntities(Environment->GetScratchBuffer().EmplaceArray<FMassEntityHandle>(Count), Count);
+		ActiveEditorEntityManager->BatchReserveEntities(ReservedEntities);
+
+		for (FMassEntityHandle ReservedEntity : ReservedEntities)
+		{
+			ReservationCallback(ReservedEntity.AsNumber());
+		}
+	}
+}
+
+void UTypedElementDatabase::BatchReserveRows(TArrayView<TypedElementDataStorage::RowHandle> ReservedRows)
+{
+	using namespace TypedElementDataStorage;
+
+	if (ActiveEditorEntityManager)
+	{
+		// Depend on the fact that a row handle is an alias for an entity within the Mass powered backend. This
+		// avoids the need for copying to a temporary array;
+		static_assert(sizeof(RowHandle) == sizeof(FMassEntityHandle),
+			"BatchReserveRows in TEDS requires the row handle and the Mass entity handle to be the same size.");
+
+		TArrayView<FMassEntityHandle>* MassHandles = reinterpret_cast<TArrayView<FMassEntityHandle>*>(&ReservedRows);
+		ActiveEditorEntityManager->BatchReserveEntities(*MassHandles);
+	}
+}
+
 TypedElementDataStorage::RowHandle UTypedElementDatabase::AddRow(TypedElementDataStorage::TableHandle Table)
 {
 	checkf(Table < Tables.Num(), TEXT("Attempting to add a row to a non-existing table."));
@@ -439,7 +472,7 @@ void UTypedElementDatabase::RemoveRow(TypedElementDataStorage::RowHandle Row)
 	FMassEntityHandle Entity = FMassEntityHandle::FromNumber(Row);
 	if (ActiveEditorEntityManager && ActiveEditorEntityManager->IsEntityValid(Entity))
 	{
-		Environment->GetIndexTable().RemoveRow(Row);
+		Environment->GetIndexTable().RemoveRow(UE::EditorDataStorage::EGlobalLockScope::Public, Row);
 		if (ActiveEditorEntityManager->IsEntityBuilt(FMassEntityHandle::FromNumber(Row)))
 		{
 			ActiveEditorEntityManager->DestroyEntity(FMassEntityHandle::FromNumber(Row));
@@ -934,23 +967,29 @@ void UTypedElementDatabase::ActivateQueries(FName ActivationName)
 
 TypedElementDataStorage::RowHandle UTypedElementDatabase::FindIndexedRow(TypedElementDataStorage::IndexHash Index) const
 {
-	return Environment->GetIndexTable().FindIndexedRow(Index);
+	return Environment->GetIndexTable().FindIndexedRow(UE::EditorDataStorage::EGlobalLockScope::Public, Index);
 }
 
 void UTypedElementDatabase::IndexRow(TypedElementDataStorage::IndexHash Index, TypedElementDataStorage::RowHandle Row)
 {
-	Environment->GetIndexTable().IndexRow(Index, Row);
+	Environment->GetIndexTable().IndexRow(UE::EditorDataStorage::EGlobalLockScope::Public, Index, Row);
+}
+
+void UTypedElementDatabase::BatchIndexRows(
+	TConstArrayView<TPair<TypedElementDataStorage::IndexHash, TypedElementDataStorage::RowHandle>> IndexRowPairs)
+{
+	Environment->GetIndexTable().BatchIndexRows(UE::EditorDataStorage::EGlobalLockScope::Public, IndexRowPairs);
 }
 
 void UTypedElementDatabase::ReindexRow(TypedElementDataStorage::IndexHash OriginalIndex, TypedElementDataStorage::IndexHash NewIndex, 
 	TypedElementDataStorage::RowHandle RowHandle)
 {
-	Environment->GetIndexTable().ReindexRow(OriginalIndex, NewIndex, RowHandle);
+	Environment->GetIndexTable().ReindexRow(UE::EditorDataStorage::EGlobalLockScope::Public, OriginalIndex, NewIndex, RowHandle);
 }
 
 void UTypedElementDatabase::RemoveIndex(TypedElementDataStorage::IndexHash Index)
 {
-	Environment->GetIndexTable().RemoveIndex(Index);
+	Environment->GetIndexTable().RemoveIndex(UE::EditorDataStorage::EGlobalLockScope::Public, Index);
 }
 
 FTypedElementOnDataStorageUpdate& UTypedElementDatabase::OnUpdate()
@@ -988,16 +1027,37 @@ void UTypedElementDatabase::ListExtensions(TFunctionRef<void(FName)> Callback) c
 
 void UTypedElementDatabase::PreparePhase(EQueryTickPhase Phase, float DeltaTime)
 {
+	using namespace UE::EditorDataStorage;
+
 	if (ActiveEditorEntityManager)
 	{
-		Environment->GetQueryStore().RunPhasePreambleQueries(*ActiveEditorEntityManager, *Environment, Phase, DeltaTime);
+		{
+			// The preamble queries are all run on the game thread. While this is true it's safe to take a global write lock.
+			// If there's a performance loss because this lock is held too long, the work in RunPhasePreambleQueries can be split
+			// into a step that runs the queries and uses a shared lock and one that executes the command buffer with an exclusive lock.
+			FScopedExclusiveLock Lock(EGlobalLockScope::Public);
+			Environment->GetQueryStore().RunPhasePreambleQueries(*ActiveEditorEntityManager, *Environment, Phase, DeltaTime);
+		}
+		// During the processing of queries no mutation can happen to the structure of the database, just fields being updated. As such
+		// it's safe to only take a shared lock
+		FGlobalLock::InternalSharedLock();
 	}
 }
 
 void UTypedElementDatabase::FinalizePhase(EQueryTickPhase Phase, float DeltaTime)
 {
+	using namespace UE::EditorDataStorage;
+
 	if (ActiveEditorEntityManager)
 	{
+		// During the processing of queries no mutation can happen to the structure of the database, just fields being updated. As such
+		// it's safe to only take a shared lock
+		FGlobalLock::InternalSharedUnlock();
+		
+		// The preamble queries are all run on the game thread. While this is true it's safe to take a global write lock.
+		// If there's a performance loss because this lock is held too long, the work in RunPhasePostambleQueries can be split
+		// into a step that runs the queries and uses a shared lock and one that executes the command buffer with an exclusive lock.
+		FScopedExclusiveLock Lock(EGlobalLockScope::Public);
 		Environment->GetQueryStore().RunPhasePostambleQueries(*ActiveEditorEntityManager, *Environment, Phase, DeltaTime);
 	}
 }
diff --git a/Engine/Plugins/Experimental/TypedElementsDataStorage/Source/TypedElementsDataStorage/Private/TypedElementDatabase.h b/Engine/Plugins/Experimental/TypedElementsDataStorage/Source/TypedElementsDataStorage/Private/TypedElementDatabase.h
--- a/Engine/Plugins/Experimental/TypedElementsDataStorage/Source/TypedElementsDataStorage/Private/TypedElementDatabase.h
+++ b/Engine/Plugins/Experimental/TypedElementsDataStorage/Source/TypedElementsDataStorage/Private/TypedElementDatabase.h
@@ -87,6 +87,8 @@ public:
 	TypedElementDataStorage::TableHandle FindTable(const FName Name) override;
 
 	TypedElementDataStorage::RowHandle ReserveRow() override;
+	void BatchReserveRows(int32 Count, TFunctionRef<void(TypedElementDataStorage::RowHandle)> ReservationCallback) override;
+	void BatchReserveRows(TArrayView<TypedElementDataStorage::RowHandle> ReservedRows) override;
 	TypedElementDataStorage::RowHandle AddRow(TypedElementDataStorage::TableHandle Table, 
 		TypedElementDataStorage::RowCreationCallbackRef OnCreated) override;
 	TypedElementDataStorage::RowHandle AddRow(TypedElementDataStorage::TableHandle Table) override;
@@ -131,6 +133,8 @@ public:
 
 	TypedElementDataStorage::RowHandle FindIndexedRow(TypedElementDataStorage::IndexHash Index) const override;
 	void IndexRow(TypedElementDataStorage::IndexHash Index, TypedElementDataStorage::RowHandle Row) override;
+	void BatchIndexRows(
+		TConstArrayView<TPair<TypedElementDataStorage::IndexHash, TypedElementDataStorage::RowHandle>> IndexRowPairs) override;
 	void ReindexRow(
 		TypedElementDataStorage::IndexHash OriginalIndex, 
 		TypedElementDataStorage::IndexHash NewIndex, 
diff --git a/Engine/Plugins/Experimental/TypedElementsDataStorage/Source/TypedElementsDataStorage/Private/TypedElementDatabaseIndexTable.cpp b/Engine/Plugins/Experimental/TypedElementsDataStorage/Source/TypedElementsDataStorage/Private/TypedElementDatabaseIndexTable.cpp
--- a/Engine/Plugins/Experimental/TypedElementsDataStorage/Source/TypedElementsDataStorage/Private/TypedElementDatabaseIndexTable.cpp
+++ b/Engine/Plugins/Experimental/TypedElementsDataStorage/Source/TypedElementsDataStorage/Private/TypedElementDatabaseIndexTable.cpp
@@ -2,39 +2,70 @@
 
 #include "TypedElementDatabaseIndexTable.h"
 
-TypedElementDataStorage::RowHandle FTypedElementDatabaseIndexTable::FindIndexedRow(TypedElementDataStorage::IndexHash Index) const
+TypedElementDataStorage::RowHandle FTypedElementDatabaseIndexTable::FindIndexedRow(
+	UE::EditorDataStorage::EGlobalLockScope LockScope, TypedElementDataStorage::IndexHash Index) const
 {
-	const TypedElementDataStorage::RowHandle* Result = IndexLookupMap.Find(Index);
-	return Result ? *Result : TypedElementDataStorage::InvalidRowHandle;
+	using namespace TypedElementDataStorage;
+	using namespace UE::EditorDataStorage;
+	
+	FScopedSharedLock Lock(LockScope);
+	
+	const RowHandle* Result = IndexLookupMap.Find(Index);
+	return Result ? *Result : InvalidRowHandle;
 }
 
-void FTypedElementDatabaseIndexTable::IndexRow(TypedElementDataStorage::IndexHash Index, TypedElementDataStorage::RowHandle Row)
+void FTypedElementDatabaseIndexTable::BatchIndexRows(UE::EditorDataStorage::EGlobalLockScope LockScope,
+	TConstArrayView<TPair<TypedElementDataStorage::IndexHash, TypedElementDataStorage::RowHandle>> IndexRowPairs)
 {
-	IndexLookupMap.Add(Index, Row);
-	ReverseIndexLookupMap.Add(Row, Index);
+	using namespace TypedElementDataStorage;
+	using namespace UE::EditorDataStorage;
+
+	FScopedExclusiveLock Lock(LockScope);
+	
+	IndexLookupMap.Reserve(IndexLookupMap.Num() + IndexRowPairs.Num());
+	ReverseIndexLookupMap.Reserve(ReverseIndexLookupMap.Num() + IndexRowPairs.Num());
+
+	for (const TPair<IndexHash, RowHandle>& IndexAndRow : IndexRowPairs)
+	{
+		IndexRowUnguarded(IndexAndRow.Key, IndexAndRow.Value);
+	}
 }
 
-void FTypedElementDatabaseIndexTable::ReindexRow(TypedElementDataStorage::IndexHash OriginalIndex, 
-	TypedElementDataStorage::IndexHash NewIndex, TypedElementDataStorage::RowHandle Row)
+void FTypedElementDatabaseIndexTable::IndexRow(UE::EditorDataStorage::EGlobalLockScope LockScope,
+	TypedElementDataStorage::IndexHash Index, TypedElementDataStorage::RowHandle Row)
 {
-	RemoveIndex(OriginalIndex);
-	IndexRow(NewIndex, Row);
+	using namespace UE::EditorDataStorage;
+
+	FScopedExclusiveLock Lock(LockScope);
+	IndexRowUnguarded(Index, Row);
 }
 
-void FTypedElementDatabaseIndexTable::RemoveIndex(TypedElementDataStorage::IndexHash Index)
+void FTypedElementDatabaseIndexTable::ReindexRow(UE::EditorDataStorage::EGlobalLockScope LockScope,
+	TypedElementDataStorage::IndexHash OriginalIndex, TypedElementDataStorage::IndexHash NewIndex, 
+	TypedElementDataStorage::RowHandle Row)
 {
-	using namespace TypedElementDataStorage;
+	using namespace UE::EditorDataStorage;
 
-	if (const RowHandle* Row = IndexLookupMap.Find(Index))
-	{
-		IndexLookupMap.Remove(Index);
-		ReverseIndexLookupMap.Remove(*Row, Index);
-	}
+	FScopedExclusiveLock Lock(LockScope);
+	
+	RemoveIndexUnguarded(OriginalIndex);
+	IndexRowUnguarded(NewIndex, Row);
+}
+
+void FTypedElementDatabaseIndexTable::RemoveIndex(UE::EditorDataStorage::EGlobalLockScope LockScope, TypedElementDataStorage::IndexHash Index)
+{
+	using namespace UE::EditorDataStorage;
+
+	FScopedExclusiveLock Lock(LockScope);
+	RemoveIndexUnguarded(Index);
 }
 
-void FTypedElementDatabaseIndexTable::RemoveRow(TypedElementDataStorage::RowHandle Row)
+void FTypedElementDatabaseIndexTable::RemoveRow(UE::EditorDataStorage::EGlobalLockScope LockScope, TypedElementDataStorage::RowHandle Row)
 {
 	using namespace TypedElementDataStorage;
+	using namespace UE::EditorDataStorage;
+
+	FScopedExclusiveLock Lock(LockScope);
 	
 	if (TMultiMap<RowHandle, IndexHash>::TKeyIterator It = ReverseIndexLookupMap.CreateKeyIterator(Row); It)
 	{
@@ -45,4 +76,22 @@ void FTypedElementDatabaseIndexTable::RemoveRow(TypedElementDataStorage::RowHand
 		} while (It);
 		ReverseIndexLookupMap.Remove(Row);
 	}
+}
+
+void FTypedElementDatabaseIndexTable::IndexRowUnguarded(TypedElementDataStorage::IndexHash Index, TypedElementDataStorage::RowHandle Row)
+{
+	IndexLookupMap.Add(Index, Row);
+	ReverseIndexLookupMap.Add(Row, Index);
+}
+
+
+void FTypedElementDatabaseIndexTable::RemoveIndexUnguarded(TypedElementDataStorage::IndexHash Index)
+{
+	using namespace TypedElementDataStorage;
+
+	if (const RowHandle* Row = IndexLookupMap.Find(Index))
+	{
+		IndexLookupMap.Remove(Index);
+		ReverseIndexLookupMap.Remove(*Row, Index);
+	}
 }
\ No newline at end of file
diff --git a/Engine/Plugins/Experimental/TypedElementsDataStorage/Source/TypedElementsDataStorage/Private/TypedElementDatabaseIndexTable.h b/Engine/Plugins/Experimental/TypedElementsDataStorage/Source/TypedElementsDataStorage/Private/TypedElementDatabaseIndexTable.h
--- a/Engine/Plugins/Experimental/TypedElementsDataStorage/Source/TypedElementsDataStorage/Private/TypedElementDatabaseIndexTable.h
+++ b/Engine/Plugins/Experimental/TypedElementsDataStorage/Source/TypedElementsDataStorage/Private/TypedElementDatabaseIndexTable.h
@@ -5,20 +5,43 @@
 #include "Containers/Map.h"
 #include "Elements/Common/TypedElementHandles.h"
 #include "Elements/Interfaces/TypedElementQueryStorageInterfaces.h"
+#include "GlobalLock.h"
+
+/**
+ * Storage for an index to row mapping.
+ * Access to the index table is thread safe and guarded by the global lock.
+ */
 
 class FTypedElementDatabaseIndexTable final
 {
 public:
-	TypedElementDataStorage::RowHandle FindIndexedRow(TypedElementDataStorage::IndexHash Index) const;
-	void IndexRow(TypedElementDataStorage::IndexHash Index, TypedElementDataStorage::RowHandle Row);
+	TypedElementDataStorage::RowHandle FindIndexedRow(
+		UE::EditorDataStorage::EGlobalLockScope LockScope,
+		TypedElementDataStorage::IndexHash Index) const;
+	void IndexRow(
+		UE::EditorDataStorage::EGlobalLockScope LockScope,
+		TypedElementDataStorage::IndexHash Index, 
+		TypedElementDataStorage::RowHandle Row);
+	void BatchIndexRows(
+		UE::EditorDataStorage::EGlobalLockScope LockScope,
+		TConstArrayView<TPair<TypedElementDataStorage::IndexHash, 
+		TypedElementDataStorage::RowHandle>> IndexRowPairs);
 	void ReindexRow(
+		UE::EditorDataStorage::EGlobalLockScope LockScope,
 		TypedElementDataStorage::IndexHash OriginalIndex,
 		TypedElementDataStorage::IndexHash NewIndex,
 		TypedElementDataStorage::RowHandle Row);
-	void RemoveIndex(TypedElementDataStorage::IndexHash Index);
-	void RemoveRow(TypedElementDataStorage::RowHandle Row);
+	void RemoveIndex(
+		UE::EditorDataStorage::EGlobalLockScope LockScope,
+		TypedElementDataStorage::IndexHash Index);
+	void RemoveRow(
+		UE::EditorDataStorage::EGlobalLockScope LockScope,
+		TypedElementDataStorage::RowHandle Row);
 
 private:
 	TMap<TypedElementDataStorage::IndexHash, TypedElementDataStorage::RowHandle> IndexLookupMap;
 	TMultiMap<TypedElementDataStorage::RowHandle, TypedElementDataStorage::IndexHash> ReverseIndexLookupMap;
+	
+	void IndexRowUnguarded(TypedElementDataStorage::IndexHash Index, TypedElementDataStorage::RowHandle Row);
+	void RemoveIndexUnguarded(TypedElementDataStorage::IndexHash Index);
 };
diff --git a/Engine/Source/Runtime/TypedElementFramework/Public/Elements/Interfaces/TypedElementDataStorageInterface.h b/Engine/Source/Runtime/TypedElementFramework/Public/Elements/Interfaces/TypedElementDataStorageInterface.h
--- a/Engine/Source/Runtime/TypedElementFramework/Public/Elements/Interfaces/TypedElementDataStorageInterface.h
+++ b/Engine/Source/Runtime/TypedElementFramework/Public/Elements/Interfaces/TypedElementDataStorageInterface.h
@@ -112,6 +112,21 @@ public:
 	 * to a table, it should still be released with RemoveRow.
 	 */
 	virtual TypedElementDataStorage::RowHandle ReserveRow() = 0;
+	/**
+	 * Reserve multiple rows at once to be assigned to a table at a later point. If multiple rows are needed, the batch version will
+	 * generally have better performance. If a row is no longer needed before it's been assigned to a table, it should still be released 
+	 * with RemoveRow.
+	 * The reservation callback will be called once per reserved row.
+	 */
+	virtual void BatchReserveRows(int32 Count, TFunctionRef<void(TypedElementDataStorage::RowHandle)> ReservationCallback) = 0;
+	/**
+	 * Reserve multiple rows at once to be assigned to a table at a later point. If multiple rows are needed, the batch version will
+	 * generally have better performance. If a row is no longer needed before it's been assigned to a table, it should still be released
+	 * with RemoveRow.
+	 * The provided range will be have its values set to the reserved row handles.
+	 */
+	virtual void BatchReserveRows(TArrayView<TypedElementDataStorage::RowHandle> ReservedRows) = 0;
+
 	/** Adds a new row to the provided table. */
 	virtual TypedElementDataStorage::RowHandle AddRow(TypedElementDataStorage::TableHandle Table) = 0;
 	/**
@@ -295,8 +310,15 @@ public:
 	 * with a single row.
 	 */
 	virtual void IndexRow(TypedElementDataStorage::IndexHash Index, TypedElementDataStorage::RowHandle Row) = 0;
+	/**
+	 * Register multiple rows under their index hash. The same row can be registered multiple times,
+	 * but an index hash can only be associated with a single row.
+	 */
+	virtual void BatchIndexRows(
+		TConstArrayView<TPair<TypedElementDataStorage::IndexHash, TypedElementDataStorage::RowHandle>> IndexRowPairs) = 0;
 	/** Updates the index of a row to a new value. Effectively this is the same as removing an index and adding a new one. */
-	virtual void ReindexRow(TypedElementDataStorage::IndexHash OriginalIndex, TypedElementDataStorage::IndexHash NewIndex, TypedElementDataStorage::RowHandle Row) = 0;
+	virtual void ReindexRow(
+		TypedElementDataStorage::IndexHash OriginalIndex, TypedElementDataStorage::IndexHash NewIndex, TypedElementDataStorage::RowHandle Row) = 0;
 	/** Removes a previously registered index hash from the index lookup table or does nothing if the hash no longer exists. */
 	virtual void RemoveIndex(TypedElementDataStorage::IndexHash Index) = 0;
 
