Simple protoype wrapper to store the asset data of the asset regsitry into teds.
For now everything is experimental so please refrain on depending on these columns and table rows at the moment.

diff --git a/Engine/Plugins/Experimental/TypedElementsDataStorage/Source/TypedElementsDataStorageAssetData/Private/Debug/TedsAssetDataDebugCommands.cpp b/Engine/Plugins/Experimental/TypedElementsDataStorage/Source/TypedElementsDataStorageAssetData/Private/Debug/TedsAssetDataDebugCommands.cpp
new file mode 100644
--- /dev/null
+++ b/Engine/Plugins/Experimental/TypedElementsDataStorage/Source/TypedElementsDataStorageAssetData/Private/Debug/TedsAssetDataDebugCommands.cpp
@@ -0,0 +1,144 @@
+// Copyright Epic Games, Inc. All Rights Reserved.
+
+#include "AssetRegistry/AssetData.h"
+#include "AssetRegistry/IAssetRegistry.h"
+#include "Containers/UnrealString.h"
+#include "Elements/Common/TypedElementHandles.h"
+#include "Elements/Framework/TypedElementIndexHasher.h"
+#include "Elements/Framework/TypedElementRegistry.h"
+#include "HAL/IConsoleManager.h"
+#include "TedsAssetDataColumns.h"
+#include "UObject/NameTypes.h"
+
+
+DEFINE_LOG_CATEGORY_STATIC(LogTEDSAssetRegistry, Log, All)
+
+static FAutoConsoleCommand CCMDTestFolderRowData(
+	TEXT("TEDS.Debug.ShowDataOfAssetFolder"),
+	TEXT("Print some debug information on the specified path."),
+	FConsoleCommandWithArgsDelegate::CreateLambda([](const TArray<FString>& InArgs)
+		{
+			const ITypedElementDataStorageInterface* Database = UTypedElementRegistry::GetInstance()->GetDataStorage();
+			IAssetRegistry& AssetRegistry = IAssetRegistry::GetChecked();
+
+			for (const FString& Path : InArgs)
+			{
+				FName PathAsName(*Path);
+				TypedElementDataStorage::RowHandle RowHandle = Database->FindIndexedRow(TypedElementDataStorage::GenerateIndexHash(PathAsName));
+
+				if (Database->IsRowAssigned(RowHandle))
+				{
+					UE_LOG(LogTEDSAssetRegistry, Display, TEXT("The path isn't indexed."));
+					return;
+				}
+
+
+				UE_LOG(LogTEDSAssetRegistry, Display, TEXT("Found some information for the path (%s) in the database."), *Path);
+
+				if (const FAssetPathColumn_Experimental* AssetPath = Database->GetColumn<FAssetPathColumn_Experimental>(RowHandle))
+				{
+					UE_LOG(LogTEDSAssetRegistry, Display, TEXT("Path stored in the database as (%s)."), *AssetPath->Path.ToString());
+				}
+
+				if (const FParentAssetPathColumn_Experimental* ParentAssetPath = Database->GetColumn<FParentAssetPathColumn_Experimental>(RowHandle))
+				{
+					if (const FAssetPathColumn_Experimental* AssetPath = Database->GetColumn<FAssetPathColumn_Experimental>(ParentAssetPath->ParentRow))
+					{
+						UE_LOG(LogTEDSAssetRegistry, Display, TEXT("	Parent Path: %s"), *AssetPath->Path.ToString());
+					}
+				}
+
+
+				if (const FChildrenAssetPathColumn_Experimental* ChildrenPath = Database->GetColumn<FChildrenAssetPathColumn_Experimental>(RowHandle))
+				{
+					UE_LOG(LogTEDSAssetRegistry, Display, TEXT("	Path as %i children"), ChildrenPath->ChildrenRows.Num());
+
+					for (TypedElementDataStorage::RowHandle Row : ChildrenPath->ChildrenRows)
+					{
+						if (const FAssetPathColumn_Experimental* AssetPath = Database->GetColumn<FAssetPathColumn_Experimental>(Row))
+						{
+							UE_LOG(LogTEDSAssetRegistry, Display, TEXT("		Children Path: %s"), *AssetPath->Path.ToString());
+						}
+					}
+				}
+
+
+				if (const FAssetsInPathColumn_Experimental* AssetInPath = Database->GetColumn<FAssetsInPathColumn_Experimental>(RowHandle))
+				{
+					UE_LOG(LogTEDSAssetRegistry, Display, TEXT("	Asset in Paths"));
+
+					for (TypedElementDataStorage::RowHandle AssetRow : AssetInPath->AssetsRow)
+					{
+						if (const FAssetDataColumn_Experimental* AssetData = Database->GetColumn<FAssetDataColumn_Experimental>(AssetRow))
+						{
+							UE_LOG(LogTEDSAssetRegistry, Display, TEXT("		Asset Name: %s"), *AssetData->AssetData.AssetName.ToString());
+						}
+						else
+						{
+							UE_LOG(LogTEDSAssetRegistry, Display, TEXT("		Asset Row pointed to stale asset."));
+						}
+					}
+				}
+
+
+				// Check for asset that haven't processed their path asset in path column yet
+				TArray<FAssetData> Assets;
+				AssetRegistry.GetAssetsByPath(PathAsName, Assets);
+				for (const FAssetData& Asset : Assets)
+				{
+					TypedElementDataStorage::RowHandle AssetRow = Database->FindIndexedRow(TypedElementDataStorage::GenerateIndexHash(Asset.GetSoftObjectPath()));
+
+					if (const FUnresolvedAssetsInPathColumn_Experimental* UnresolvedAssetsInPathColumn = Database->GetColumn<FUnresolvedAssetsInPathColumn_Experimental>(AssetRow))
+					{
+							UE_LOG(LogTEDSAssetRegistry, Display, TEXT("		Unresolved asset in path asset. Asset Name: %s"), *Asset.AssetName.ToString());
+					}
+				}
+			}
+		}));
+
+
+static FAutoConsoleCommand CCMDTestFolderAssetRegistryData(
+	TEXT("TEDS.Debug.ShowAssetRegistryDataOfFolder"),
+	TEXT("Print some debug information on the specified path."),
+	FConsoleCommandWithArgsDelegate::CreateLambda([](const TArray<FString>& InArgs)
+		{
+			IAssetRegistry& AssetRegistry = IAssetRegistry::GetChecked();
+
+			FNameBuilder NameBuilder;
+			for (const FString& Path : InArgs)
+			{
+				UE_LOG(LogTEDSAssetRegistry, Display, TEXT("Displaying asset registry info on path (%s)"), *Path);
+				const FName PathAsName(*Path);
+				AssetRegistry.EnumerateSubPaths(PathAsName, [&NameBuilder](FName InPath)
+					{
+						InPath.AppendString(NameBuilder);
+						UE_LOG(LogTEDSAssetRegistry, Display, TEXT("	Children Path: %s"), *NameBuilder);
+						NameBuilder.Reset();
+						return true;
+					}, false);
+
+
+				TArray<FAssetData> Assets;
+				AssetRegistry.GetAssetsByPath(PathAsName, Assets);
+
+				if (!Assets.IsEmpty())
+				{ 
+					UE_LOG(LogTEDSAssetRegistry, Display, TEXT("	Asset in Path"));
+
+					for (const FAssetData& Asset : Assets)
+					{
+						Asset.AssetName.AppendString(NameBuilder);
+						UE_LOG(LogTEDSAssetRegistry, Display, TEXT("		Asset Name: %s"), *NameBuilder);
+						NameBuilder.Reset();
+
+						Asset.GetFullName(NameBuilder);
+						UE_LOG(LogTEDSAssetRegistry, Display, TEXT("		Asset Full Name: %s"), *NameBuilder);
+						NameBuilder.Reset();
+
+						Asset.PackageName.AppendString(NameBuilder);
+						UE_LOG(LogTEDSAssetRegistry, Display, TEXT("		Asset Reported Package Path: %s"), *NameBuilder);
+						NameBuilder.Reset();
+					}
+				}
+			}
+		}));
\ No newline at end of file
diff --git a/Engine/Plugins/Experimental/TypedElementsDataStorage/Source/TypedElementsDataStorageAssetData/Private/TedsAssetData.cpp b/Engine/Plugins/Experimental/TypedElementsDataStorage/Source/TypedElementsDataStorageAssetData/Private/TedsAssetData.cpp
new file mode 100644
--- /dev/null
+++ b/Engine/Plugins/Experimental/TypedElementsDataStorage/Source/TypedElementsDataStorageAssetData/Private/TedsAssetData.cpp
@@ -0,0 +1,678 @@
+// Copyright Epic Games, Inc. All Rights Reserved.
+
+#include "TedsAssetData.h"
+
+#include "TedsAssetDataColumns.h"
+
+#include "AssetRegistry/IAssetRegistry.h"
+#include "Async/ParallelFor.h"
+#include "Containers/Array.h"
+#include "Containers/ChunkedArray.h"
+#include "Elements/Common/TypedElementQueryTypes.h"
+#include "Elements/Framework/TypedElementIndexHasher.h"
+#include "Elements/Framework/TypedElementQueryBuilder.h"
+#include "Elements/Interfaces/TypedElementDataStorageInterface.h"
+#include "Engine/Engine.h"
+
+#define TRACK_TEDSASSETDATA_MEMORY 0
+
+namespace UE::TypedElementsDataStorageAssetData::Private
+{
+	
+constexpr int32 ParallelForMinBatchSize = 1024 * 4;
+
+struct FPopulateAssetDataRowArgs
+{
+	FAssetData AssetData;
+
+	TypedElementDataStorage::IndexHash ObjectPathHash;
+	TypedElementDataStorage::IndexHash PathHash;
+	TypedElementDataStorage::RowHandle PathRow = TypedElementDataStorage::InvalidRowHandle;
+};
+
+// Only safe if the GT is blocked during the operation
+template<typename TAssetData>
+FPopulateAssetDataRowArgs ThreadSafe_PopulateAssetDataTableRow(TAssetData&& InAssetData, const ITypedElementDataStorageInterface& Database)
+{
+	FPopulateAssetDataRowArgs Output;
+	Output.ObjectPathHash = TypedElementDataStorage::GenerateIndexHash(InAssetData.GetSoftObjectPath());
+
+	// Looks safe but might not be depending on the implementation of the database
+	if (Database.IsRowAssigned(Database.FindIndexedRow(Output.ObjectPathHash)))
+	{
+		// No need to initialize the rest of the row here. The invalid's asset data will be used as flag to skip the data generated here.
+		return Output;
+	}
+
+	Output.PathHash = TypedElementDataStorage::GenerateIndexHash(InAssetData.PackagePath);
+
+	// Looks safe but might not be depending on the implementation of the database
+	Output.PathRow = Database.FindIndexedRow(Output.PathHash);
+	Output.AssetData = MoveTempIfPossible(InAssetData);
+
+	return Output;
+}
+
+
+void PopulateAssetDataTableRow(FPopulateAssetDataRowArgs&& InAssetDataRowArgs, ITypedElementDataStorageInterface& Database, TypedElementDataStorage::RowHandle RowHandle)
+{
+	if (InAssetDataRowArgs.PathRow !=  TypedElementDataStorage::InvalidRowHandle)
+	{
+		Database.GetColumn<FAssetsInPathColumn_Experimental>(InAssetDataRowArgs.PathRow)->AssetsRow.Add(RowHandle);
+	}
+	else
+	{
+		FUnresolvedAssetsInPathColumn_Experimental UnresolvedAssetsInPathColumn;
+		UnresolvedAssetsInPathColumn.Hash = InAssetDataRowArgs.PathHash;
+		Database.AddColumn<FUnresolvedAssetsInPathColumn_Experimental>(RowHandle, MoveTemp(UnresolvedAssetsInPathColumn));
+	}
+
+	Database.GetColumn<FAssetDataColumn_Experimental>(RowHandle)->AssetData =  MoveTemp(InAssetDataRowArgs.AssetData);
+}
+
+
+struct FPopulatePathRowArgs
+{
+	FName AssetRegistryPath;
+	TypedElementDataStorage::IndexHash AssetRegistryPathHash;
+	TypedElementDataStorage::IndexHash ParentAssetRegistryPathHash;
+	uint32 PathDepth;
+
+	operator bool() const 
+	{
+		return !AssetRegistryPath.IsNone();
+	}
+
+	void MarkAsInvalid()
+	{
+		AssetRegistryPath = FName();
+	}
+};
+
+void GetPathDepthAndParentFolderIndex(FStringView Path, uint32& OutDepth, int32& OutParentFolderIndex)
+{
+	OutDepth = 0;
+	OutParentFolderIndex = INDEX_NONE;
+
+	if (Path.Len() > 1)
+	{
+		OutParentFolderIndex = 1;
+		++OutDepth;
+	}
+
+	// Skip the first '/'
+	for (int32 Index = 1; Index < Path.Len(); ++Index)
+	{
+		if (Path[Index] == TEXT('/'))
+		{
+			++OutDepth;
+			OutParentFolderIndex = Index;
+		}
+	}
+}
+
+// Only thread safe if the game thread is blocked 
+FPopulatePathRowArgs ThreadSafe_PopulatePathRowArgs(TypedElementDataStorage::IndexHash AssetRegistryPathHash, FName InAssetRegistryPath, FStringView PathAsString)
+{
+	TypedElementDataStorage::IndexHash ParentAssetRegistryPathHash = TypedElementDataStorage::InvalidRowHandle;
+	int32 CharacterIndex;
+	uint32 Depth;
+	GetPathDepthAndParentFolderIndex(PathAsString, Depth, CharacterIndex);
+	if (CharacterIndex != INDEX_NONE)
+	{
+		FStringView ParentPath = PathAsString.Left(CharacterIndex);
+		ParentAssetRegistryPathHash = TypedElementDataStorage::GenerateIndexHash(FName(ParentPath));
+	}
+
+	FPopulatePathRowArgs Args;
+	Args.AssetRegistryPath = InAssetRegistryPath;
+	Args.AssetRegistryPathHash = AssetRegistryPathHash;
+	Args.ParentAssetRegistryPathHash = ParentAssetRegistryPathHash;
+	Args.PathDepth = Depth;
+
+	return Args;
+}
+
+void PopulatePathDataTableRow(FPopulatePathRowArgs&& InPopulatePathRowArgs, ITypedElementDataStorageInterface& Database, TypedElementDataStorage::RowHandle RowHandle)
+{
+	if (InPopulatePathRowArgs.ParentAssetRegistryPathHash != TypedElementDataStorage::InvalidRowHandle)
+	{
+		const TypedElementDataStorage::RowHandle ParentRow = Database.FindIndexedRow(InPopulatePathRowArgs.ParentAssetRegistryPathHash);
+		if (Database.IsRowAssigned(ParentRow))
+		{
+			Database.GetColumn<FChildrenAssetPathColumn_Experimental>(ParentRow)->ChildrenRows.Add(RowHandle);
+			Database.GetColumn<FParentAssetPathColumn_Experimental>(RowHandle)->ParentRow = ParentRow;
+		}
+		else 
+		{
+			// If we were unlucky we may be missing the parent path data.This should resolve it self during a latter event on paths added event.
+			FUnresolvedParentAssetPathColumn_Experimental UnresolvedParentRow;
+			UnresolvedParentRow.Hash = InPopulatePathRowArgs.ParentAssetRegistryPathHash;
+			Database.AddColumn(RowHandle, MoveTemp(UnresolvedParentRow));
+		}
+	}
+
+	Database.GetColumn<FAssetPathColumn_Experimental>(RowHandle)->Path = InPopulatePathRowArgs.AssetRegistryPath;
+}
+
+
+FTedsAssetData::FTedsAssetData(ITypedElementDataStorageInterface& InDatabase)
+	: Database(InDatabase)
+{
+	TRACE_CPUPROFILER_EVENT_SCOPE(FTedsAssetData::FTedsAssetData);
+
+#if TRACK_TEDSASSETDATA_MEMORY
+	LLM_SCOPE_BYNAME(TEXT("FTedsAssetData"))
+#endif
+
+	// Register to events from asset registry
+	IAssetRegistry& AssetRegistry = IAssetRegistry::GetChecked();
+
+	AssetRegistry.OnAssetsAdded().AddRaw(this, &FTedsAssetData::OnAssetsAdded);
+	AssetRegistry.OnAssetsRemoved().AddRaw(this, &FTedsAssetData::OnAssetsRemoved);
+	AssetRegistry.OnAssetsUpdated().AddRaw(this, &FTedsAssetData::OnAssetsUpdated);
+	AssetRegistry.OnAssetRenamed().AddRaw(this, &FTedsAssetData::OnAssetRenamed);
+	AssetRegistry.OnAssetsUpdatedOnDisk().AddRaw(this, &FTedsAssetData::OnAssetsUpdatedOnDisk);
+	AssetRegistry.OnPathsAdded().AddRaw(this, &FTedsAssetData::OnPathsAdded);
+	AssetRegistry.OnPathsRemoved().AddRaw(this, &FTedsAssetData::OnPathsRemoved);
+
+	// Register data types to TEDS
+	PathsTable = Database.FindTable(FName(TEXT("Editor_AssetRegistryPathsTable")));
+	if (PathsTable == TypedElementDataStorage::InvalidTableHandle)
+	{
+		PathsTable = Database.RegisterTable<FAssetPathColumn_Experimental, FChildrenAssetPathColumn_Experimental, FParentAssetPathColumn_Experimental, FAssetsInPathColumn_Experimental>(FName(TEXT("Editor_AssetRegistryPathsTable")));
+	}
+
+	AssetsDataTable = Database.FindTable(FName(TEXT("Editor_AssetRegistryAssetDataTable")));
+	if (AssetsDataTable == TypedElementDataStorage::InvalidTableHandle)
+	{
+		AssetsDataTable = Database.RegisterTable<FAssetDataColumn_Experimental>(FName(TEXT("Editor_AssetRegistryAssetDataTable")));
+	}
+
+	UpdateAssetsInPath = Database.RegisterQuery(
+		TypedElementQueryBuilder::Select()
+			.ReadWrite<FAssetsInPathColumn_Experimental>()
+			.Compile()
+		);
+
+	ResolveMissingAssetInPathQuery = Database.RegisterQuery(
+		TypedElementQueryBuilder::Select(
+			TEXT("Resolve Missing Asset In Path"),
+			TypedElementQueryBuilder::FProcessor(TypedElementDataStorage::EQueryTickPhase::FrameEnd, Database.GetQueryTickGroupName(TypedElementDataStorage::EQueryTickGroups::Default)),
+			[this](TypedElementDataStorage::IQueryContext& Context, TypedElementDataStorage::RowHandle Row, const FUnresolvedAssetsInPathColumn_Experimental& UnresolvedAssetPath)
+			{
+#if TRACK_TEDSASSETDATA_MEMORY
+				LLM_SCOPE_BYNAME(TEXT("FTedsAssetData"))
+#endif
+
+				const TypedElementDataStorage::RowHandle PathRow = Context.FindIndexedRow(UnresolvedAssetPath.Hash);
+				if (Context.IsRowAssigned(PathRow))
+				{
+					Context.RemoveColumns<FUnresolvedAssetsInPathColumn_Experimental>(Row);
+					Context.RunSubquery(0, PathRow, 
+							TypedElementQueryBuilder::CreateSubqueryCallbackBinding([Row](FAssetsInPathColumn_Experimental& AssetsInPath)
+							{
+#if TRACK_TEDSASSETDATA_MEMORY
+								LLM_SCOPE_BYNAME(TEXT("FTedsAssetData"))
+#endif
+
+								AssetsInPath.AssetsRow.Add(Row);
+							})
+						);
+				}
+			}
+		)
+		.DependsOn()
+			.SubQuery(UpdateAssetsInPath)
+		.Compile()
+	);
+
+	UpdateParentToChildrenAssetPathQuery = Database.RegisterQuery(
+		TypedElementQueryBuilder::Select()
+			.ReadWrite<FChildrenAssetPathColumn_Experimental>()
+			.Compile()
+		);
+
+	ResolveMissingParentPathQuery = Database.RegisterQuery(
+		TypedElementQueryBuilder::Select(
+			TEXT("Resolve Missing Parent Path Row"),
+			TypedElementQueryBuilder::FProcessor(TypedElementDataStorage::EQueryTickPhase::FrameEnd, Database.GetQueryTickGroupName(TypedElementDataStorage::EQueryTickGroups::Default)),
+			[this](TypedElementDataStorage::IQueryContext& Context, TypedElementDataStorage::RowHandle Row, const FUnresolvedParentAssetPathColumn_Experimental& UnresolvedParentAssetPath, FParentAssetPathColumn_Experimental& ParentAssetPathColumn)
+			{
+#if TRACK_TEDSASSETDATA_MEMORY
+				LLM_SCOPE_BYNAME(TEXT("FTedsAssetData"))
+#endif
+
+				const TypedElementDataStorage::RowHandle ParentPathRow = Context.FindIndexedRow(UnresolvedParentAssetPath.Hash);
+				if (Context.IsRowAssigned(ParentPathRow))
+				{
+					Context.RemoveColumns<FUnresolvedParentAssetPathColumn_Experimental>(Row);
+					ParentAssetPathColumn.ParentRow = ParentPathRow;
+
+					Context.RunSubquery(0, ParentPathRow, 
+							TypedElementQueryBuilder::CreateSubqueryCallbackBinding([Row](FChildrenAssetPathColumn_Experimental& ChildrenPathColumn)
+							{
+#if TRACK_TEDSASSETDATA_MEMORY
+								LLM_SCOPE_BYNAME(TEXT("FTedsAssetData"))
+#endif
+								ChildrenPathColumn.ChildrenRows.Add(Row);
+							})
+						);
+				}
+			}
+		)
+		.DependsOn()
+			.SubQuery(UpdateParentToChildrenAssetPathQuery)
+		.Compile()
+	);
+
+	// Init with the data existing at moment in asset registry
+
+	TArray<FAssetData> AssetsData;
+	AssetRegistry.GetAllAssets(AssetsData);
+
+	TChunkedArray<FName> CachedPaths;
+	AssetRegistry.EnumerateAllCachedPaths([&CachedPaths](FName Name)
+		{
+			CachedPaths.AddElement(Name);
+			return true;
+		});
+
+	TArray<FNameBuilder> NameBuilders;
+	TArray<FPopulatePathRowArgs> PopulatePathRowArgs;
+	PopulatePathRowArgs.AddUninitialized(CachedPaths.Num());
+
+	// Prepare Path Rows
+	ParallelForWithTaskContext(TEXT("Populating TEDS Asset Registry Path"), NameBuilders, CachedPaths.Num(),ParallelForMinBatchSize, [&PopulatePathRowArgs, &CachedPaths, this](FNameBuilder& NameBuilder, int32 Index)
+		{
+			FName Path = CachedPaths[Index];
+			Path.ToString(NameBuilder);
+			PopulatePathRowArgs[Index] = ThreadSafe_PopulatePathRowArgs(TypedElementDataStorage::GenerateIndexHash(Path), Path, NameBuilder);
+		});
+
+	TArray<TypedElementDataStorage::RowHandle> ReservedRows;
+	ReservedRows.AddUninitialized(PopulatePathRowArgs.Num() + AssetsData.Num());
+	Database.BatchReserveRows(ReservedRows);
+
+	TArray<TPair<TypedElementDataStorage::IndexHash, TypedElementDataStorage::RowHandle>> IndexesToReservedRows;
+	IndexesToReservedRows.AddUninitialized(ReservedRows.Num());
+
+	// Index Reserved Path Rows
+	TConstArrayView<TypedElementDataStorage::RowHandle> ReservedPopulatePathRows(ReservedRows.GetData(), PopulatePathRowArgs.Num());
+	TArrayView<TPair<TypedElementDataStorage::IndexHash, TypedElementDataStorage::RowHandle>> IndexesToReservedPathRows(IndexesToReservedRows.GetData(), PopulatePathRowArgs.Num());
+	ParallelFor(TEXT("Populating TEDS Asset Registry Path Data Indexes"), ReservedPopulatePathRows.Num(), ParallelForMinBatchSize, [&IndexesToReservedPathRows, &ReservedPopulatePathRows, &PopulatePathRowArgs](int32 Index)
+		{
+			IndexesToReservedPathRows[Index] = TPair<TypedElementDataStorage::IndexHash, TypedElementDataStorage::RowHandle>(PopulatePathRowArgs[Index].AssetRegistryPathHash, ReservedPopulatePathRows[Index]);
+		});
+	Database.BatchIndexRows(IndexesToReservedPathRows);
+
+	// Populate Path Rows
+	Database.BatchAddRow(PathsTable, ReservedPopulatePathRows, [PathRowArgs = MoveTemp(PopulatePathRowArgs), Index = 0, this](TypedElementDataStorage::RowHandle InRowHandle) mutable
+		{
+			PopulatePathDataTableRow(MoveTemp(PathRowArgs[Index]), Database, InRowHandle);
+			++Index;
+		});
+
+	// Prepare Asset Data Rows
+	TArray<FPopulateAssetDataRowArgs> PopulateAssetDataRowArgs;
+	PopulateAssetDataRowArgs.AddDefaulted(AssetsData.Num());
+	ParallelFor(TEXT("Populating TEDS Asset Registry Asset Data"), PopulateAssetDataRowArgs.Num(), ParallelForMinBatchSize, [&PopulateAssetDataRowArgs, &AssetsData, this](int32 Index)
+		{
+			PopulateAssetDataRowArgs[Index] = ThreadSafe_PopulateAssetDataTableRow(MoveTemp(AssetsData[Index]), Database);
+		});
+
+	// Index Reserved Asset Data Rows
+	TConstArrayView<TypedElementDataStorage::RowHandle> ReservedAssetDataRows(ReservedRows.GetData() + ReservedPopulatePathRows.Num(), AssetsData.Num());
+	TArrayView<TPair<TypedElementDataStorage::IndexHash, TypedElementDataStorage::RowHandle>>  IndexesToResevedAssetRows(IndexesToReservedRows.GetData() + ReservedPopulatePathRows.Num(), ReservedAssetDataRows.Num());
+	ParallelFor(TEXT("Populating TEDS Asset Registry Asset Data Indexes"), PopulateAssetDataRowArgs.Num(), ParallelForMinBatchSize, [&PopulateAssetDataRowArgs, &ReservedAssetDataRows, &IndexesToResevedAssetRows](int32 Index)
+		{
+			IndexesToResevedAssetRows[Index] = TPair<TypedElementDataStorage::IndexHash, TypedElementDataStorage::RowHandle>(PopulateAssetDataRowArgs[Index].ObjectPathHash, ReservedAssetDataRows[Index]);
+		});
+	Database.BatchIndexRows(IndexesToResevedAssetRows);
+
+	// Populate Asset Rows
+	Database.BatchAddRow(AssetsDataTable, ReservedAssetDataRows,  [AssetDataRowArgs = MoveTemp(PopulateAssetDataRowArgs), Index = 0, this](TypedElementDataStorage::RowHandle InRowHandle) mutable
+		{
+			PopulateAssetDataTableRow(MoveTemp(AssetDataRowArgs[Index]), Database, InRowHandle);
+			++Index;
+		});
+}
+
+FTedsAssetData::~FTedsAssetData()
+{
+	TRACE_CPUPROFILER_EVENT_SCOPE(FTedsAssetData::~FTedsAssetData);
+
+#if TRACK_TEDSASSETDATA_MEMORY
+	LLM_SCOPE_BYNAME(TEXT("FTedsAssetData"))
+#endif
+
+	if (IAssetRegistry* AssetRegistry = IAssetRegistry::Get())
+	{
+		Database.UnregisterQuery(ResolveMissingParentPathQuery);
+		Database.UnregisterQuery(UpdateParentToChildrenAssetPathQuery);
+		Database.UnregisterQuery(ResolveMissingAssetInPathQuery);
+		Database.UnregisterQuery(UpdateAssetsInPath);
+
+
+		AssetRegistry->OnAssetsAdded().RemoveAll(this);
+		AssetRegistry->OnAssetsRemoved().RemoveAll(this);
+		AssetRegistry->OnAssetsUpdated().RemoveAll(this);
+		AssetRegistry->OnAssetsUpdatedOnDisk().RemoveAll(this);
+		AssetRegistry->OnAssetRenamed().RemoveAll(this);
+		AssetRegistry->OnPathsAdded().RemoveAll(this);
+		AssetRegistry->OnPathsRemoved().RemoveAll(this);
+	
+
+		// Not needed on a editor shut down
+		if (!IsEngineExitRequested())
+		{
+			AssetRegistry->EnumerateAllCachedPaths([this](FName InPath)
+				{
+					const TypedElementDataStorage::IndexHash PathHash = TypedElementDataStorage::GenerateIndexHash(InPath);
+					const TypedElementDataStorage::RowHandle Row = Database.FindIndexedRow(PathHash);
+					Database.RemoveRow(Row);
+					Database.RemoveIndex(PathHash);
+					return true;
+				});
+
+			AssetRegistry->EnumerateAllAssets([this](const FAssetData& InAssetData)
+				{
+					const TypedElementDataStorage::IndexHash AssetPathHash = TypedElementDataStorage::GenerateIndexHash(InAssetData.GetSoftObjectPath());
+					const TypedElementDataStorage::RowHandle Row = Database.FindIndexedRow(AssetPathHash);
+					Database.RemoveRow(Row);
+					Database.RemoveIndex(AssetPathHash);
+					return true;
+				});
+		}
+	}
+}
+
+void FTedsAssetData::ProcessAllEvents()
+{
+	if (IAssetRegistry* AssetRegistry = IAssetRegistry::Get())
+	{
+		AssetRegistry->Tick(-1.f);
+	}
+}
+
+void FTedsAssetData::OnAssetsAdded(TConstArrayView<FAssetData> InAssetsAdded)
+{
+	TRACE_CPUPROFILER_EVENT_SCOPE(FTedsAssetData::OnAssetsAdded);
+
+#if TRACK_TEDSASSETDATA_MEMORY
+	LLM_SCOPE_BYNAME(TEXT("FTedsAssetData"))
+#endif
+
+	TArray<int32> Contexts;
+	TArray<FPopulateAssetDataRowArgs> PopulateRowArgs;
+	PopulateRowArgs.AddDefaulted(InAssetsAdded.Num());
+
+	UE::AssetRegistry::FFiltering::InitializeShouldSkipAsset();
+
+	ParallelForWithTaskContext(TEXT("Populating TEDS Asset Registry Asset Data"), Contexts, InAssetsAdded.Num(), ParallelForMinBatchSize, [&PopulateRowArgs, &InAssetsAdded, this](int32& ValidAssetCount, int32 Index)
+		{
+			FPopulateAssetDataRowArgs RowArgs;
+			const FAssetData& AssetData = InAssetsAdded[Index];
+
+			if (!UE::AssetRegistry::FFiltering::ShouldSkipAsset(AssetData.AssetClassPath, AssetData.PackageFlags))
+			{
+				RowArgs = ThreadSafe_PopulateAssetDataTableRow(AssetData, Database);
+				if (RowArgs.AssetData.IsValid())
+				{
+					++ValidAssetCount;
+				}
+			}
+			PopulateRowArgs[Index] = MoveTemp(RowArgs);
+		});
+
+
+	int32 NewRowsCount = 0;
+	for (int32 Context : Contexts)
+	{
+		NewRowsCount += Context;
+	}
+
+	int32 Index = 0;
+	if (NewRowsCount > 0)
+	{
+		TArray<TPair<TypedElementDataStorage::IndexHash, TypedElementDataStorage::RowHandle>> IndexToRow;
+		IndexToRow.Reserve(NewRowsCount);
+		Database.BatchAddRow(AssetsDataTable, NewRowsCount, [RowArgs = MoveTemp(PopulateRowArgs), Index = 0, &IndexToRow, this](TypedElementDataStorage::RowHandle InRowHandle) mutable
+			{
+				FPopulateAssetDataRowArgs ARowArgs = MoveTemp(RowArgs[Index]);
+				while (!ARowArgs.AssetData.IsValid())
+				{
+					++Index;
+					ARowArgs = MoveTemp(RowArgs[Index]);
+				}
+
+				IndexToRow.Emplace(ARowArgs.ObjectPathHash, InRowHandle);
+				PopulateAssetDataTableRow(MoveTemp(ARowArgs), Database, InRowHandle);
+				++Index;
+			});
+
+		Database.BatchIndexRows(IndexToRow);
+	}
+}
+
+void FTedsAssetData::OnAssetsRemoved(TConstArrayView<FAssetData> InAssetsRemoved)
+{
+	TRACE_CPUPROFILER_EVENT_SCOPE(FTedsAssetData::OnAssetsRemoved);
+
+#if TRACK_TEDSASSETDATA_MEMORY
+	LLM_SCOPE_BYNAME(TEXT("FTedsAssetData"))
+#endif
+
+	for (const FAssetData& Asset : InAssetsRemoved)
+	{
+		const TypedElementDataStorage::IndexHash AssetHash = TypedElementDataStorage::GenerateIndexHash(Asset.GetSoftObjectPath());
+		const TypedElementDataStorage::RowHandle AssetRow = Database.FindIndexedRow(AssetHash);
+		if (Database.IsRowAssigned(AssetRow))
+		{
+			if (const FAssetDataColumn_Experimental* AssetDataColunm = Database.GetColumn<FAssetDataColumn_Experimental>(AssetRow))
+			{
+				const TypedElementDataStorage::IndexHash FolderPathHash = TypedElementDataStorage::GenerateIndexHash(AssetDataColunm->AssetData.PackagePath);
+				const TypedElementDataStorage::RowHandle FolderRow = Database.FindIndexedRow(FolderPathHash);
+				if (FAssetsInPathColumn_Experimental* AssetInFolder = Database.GetColumn<FAssetsInPathColumn_Experimental>(FolderRow))
+				{
+					AssetInFolder->AssetsRow.Remove(AssetRow);
+				}
+			}
+			Database.RemoveRow(Database.FindIndexedRow(AssetHash));
+			Database.RemoveIndex(AssetHash);
+		}
+	}
+}
+
+void FTedsAssetData::OnAssetsUpdated(TConstArrayView<FAssetData> InAssetsUpdated)
+{
+	TRACE_CPUPROFILER_EVENT_SCOPE(FTedsAssetData::OnAssetsUpdated);
+
+#if TRACK_TEDSASSETDATA_MEMORY
+	LLM_SCOPE_BYNAME(TEXT("FTedsAssetData"))
+#endif
+
+	for (const FAssetData& Asset : InAssetsUpdated)
+	{
+		const TypedElementDataStorage::IndexHash AssetHash = TypedElementDataStorage::GenerateIndexHash(Asset.GetSoftObjectPath());
+		const TypedElementDataStorage::RowHandle Row = Database.FindIndexedRow(AssetHash);
+		if (Database.IsRowAssigned(Row))
+		{
+			Database.GetColumn<FAssetDataColumn_Experimental>(Row)->AssetData = Asset;
+		}
+	}
+}
+
+void FTedsAssetData::OnAssetsUpdatedOnDisk(TConstArrayView<FAssetData> InAssetsUpdated)
+{
+	TRACE_CPUPROFILER_EVENT_SCOPE(FTedsAssetData::OnAssetsUpdatedOnDisk);
+
+#if TRACK_TEDSASSETDATA_MEMORY
+	LLM_SCOPE_BYNAME(TEXT("FTedsAssetData"))
+#endif
+
+	for (const FAssetData& Asset : InAssetsUpdated)
+	{
+		const TypedElementDataStorage::IndexHash AssetHash = TypedElementDataStorage::GenerateIndexHash(Asset.GetSoftObjectPath());
+		const TypedElementDataStorage::RowHandle Row = Database.FindIndexedRow(AssetHash);
+		if (Database.IsRowAssigned(Row))
+		{
+			Database.GetColumn<FAssetDataColumn_Experimental>(Row)->AssetData = Asset;
+		}
+	}
+}
+
+void FTedsAssetData::OnAssetRenamed(const FAssetData& InAsset, const FString& InOldObjectPath)
+{
+	TRACE_CPUPROFILER_EVENT_SCOPE(FTedsAssetData::OnAssetRenamed);
+
+#if TRACK_TEDSASSETDATA_MEMORY
+	LLM_SCOPE_BYNAME(TEXT("FTedsAssetData"))
+#endif
+
+	const TypedElementDataStorage::IndexHash NewAssetHash = TypedElementDataStorage::GenerateIndexHash(InAsset.GetSoftObjectPath());
+	const TypedElementDataStorage::IndexHash OldAssetHash = TypedElementDataStorage::GenerateIndexHash(FSoftObjectPath(InOldObjectPath));
+	const TypedElementDataStorage::RowHandle Row = Database.FindIndexedRow(OldAssetHash);
+	if (Database.IsRowAssigned(Row))
+	{
+		Database.GetColumn<FAssetDataColumn_Experimental>(Row)->AssetData = InAsset;
+
+		// Update the asset in folder columns
+		const TypedElementDataStorage::IndexHash NewFolderHash = TypedElementDataStorage::GenerateIndexHash(InAsset.PackagePath);
+		FStringView OldPackagePath(InOldObjectPath);
+
+		{
+			int32 CharacterIndex = 0;
+			OldPackagePath.FindLastChar(TEXT('/'), CharacterIndex);
+			OldPackagePath.LeftInline(CharacterIndex);
+		}
+
+		const TypedElementDataStorage::IndexHash OldFolderHash = TypedElementDataStorage::GenerateIndexHash(FName(OldPackagePath));
+
+		if (NewFolderHash != OldFolderHash)
+		{
+			const TypedElementDataStorage::RowHandle NewPathRow = Database.FindIndexedRow(NewFolderHash);
+			if (NewPathRow !=  TypedElementDataStorage::InvalidRowHandle)
+			{
+				Database.GetColumn<FAssetsInPathColumn_Experimental>(NewPathRow)->AssetsRow.Add(Row);
+			}
+			else
+			{
+				FUnresolvedAssetsInPathColumn_Experimental UnresolvedAssetsInPathColumn;
+				UnresolvedAssetsInPathColumn.Hash = NewFolderHash;
+				Database.AddColumn<FUnresolvedAssetsInPathColumn_Experimental>(Row, MoveTemp(UnresolvedAssetsInPathColumn));
+			}
+
+			const TypedElementDataStorage::RowHandle OldPathRow = Database.FindIndexedRow(OldFolderHash);
+			if (NewPathRow !=  TypedElementDataStorage::InvalidRowHandle)
+			{
+				Database.GetColumn<FAssetsInPathColumn_Experimental>(NewPathRow)->AssetsRow.Remove(Row);
+			}
+			else
+			{
+				Database.RemoveColumn<FUnresolvedAssetsInPathColumn_Experimental>(Row);
+			}
+		}
+
+		Database.ReindexRow(OldAssetHash, NewAssetHash, Row);
+	}
+}
+
+void FTedsAssetData::OnPathsAdded(TConstArrayView<FStringView> InPathsAdded)
+{
+	TRACE_CPUPROFILER_EVENT_SCOPE(FTedsAssetData::OnPathsAdded);
+
+#if TRACK_TEDSASSETDATA_MEMORY
+	LLM_SCOPE_BYNAME(TEXT("FTedsAssetData"));
+#endif
+
+	TArray<int32> PreparePathRowContexts;
+	TArray<FPopulatePathRowArgs> PopulateRowArgs;
+	PopulateRowArgs.AddUninitialized(InPathsAdded.Num());
+
+	ParallelForWithTaskContext(TEXT("Populating TEDS Asset Registry Path"), PreparePathRowContexts, InPathsAdded.Num(), ParallelForMinBatchSize, [&PopulateRowArgs, &InPathsAdded, this](int32& WorkerValidCount, int32 Index)
+		{
+			FStringView Path = InPathsAdded[Index];
+			FName PathName(Path);
+			const TypedElementDataStorage::IndexHash AssetRegistryPathHash = TypedElementDataStorage::GenerateIndexHash(PathName);
+			FPopulatePathRowArgs RowArgs;
+
+			if (Database.FindIndexedRow(AssetRegistryPathHash) != TypedElementDataStorage::InvalidRowHandle)
+			{
+				RowArgs.MarkAsInvalid();
+			}
+			else
+			{
+				RowArgs = ThreadSafe_PopulatePathRowArgs(AssetRegistryPathHash, PathName, Path);
+				++WorkerValidCount;
+			}
+
+			PopulateRowArgs[Index] = MoveTemp(RowArgs);
+		});
+
+	int32 NewRowsCount = 0;
+	for (const int32& Context : PreparePathRowContexts)
+	{
+		NewRowsCount += Context;
+	}
+
+
+	if (NewRowsCount > 0)
+	{
+		TArray<TypedElementDataStorage::RowHandle> ReservedRow;
+		ReservedRow.AddUninitialized(NewRowsCount);
+		Database.BatchReserveRows(ReservedRow);
+
+		TArray<TPair<TypedElementDataStorage::IndexHash, TypedElementDataStorage::RowHandle>> IndexesAndRows;
+		IndexesAndRows.Reserve(NewRowsCount);
+
+		int32 RowCount = 0;
+		for (int32 Index = 0; Index < NewRowsCount; ++Index)
+		{
+			const FPopulatePathRowArgs* RowArgs = &PopulateRowArgs[RowCount];
+			while (!RowArgs)
+			{
+				++RowCount;
+				RowArgs = &PopulateRowArgs[RowCount];
+			}
+
+			IndexesAndRows.Emplace(RowArgs->AssetRegistryPathHash, ReservedRow[Index]);
+			++RowCount;
+		}
+
+		Database.BatchIndexRows(IndexesAndRows);
+
+		int32 Index = 0;
+		Database.BatchAddRow(PathsTable, ReservedRow, [&PopulateRowArgs, &Index, this](TypedElementDataStorage::RowHandle InRowHandle)
+			{
+				FPopulatePathRowArgs RowArgs = MoveTemp(PopulateRowArgs[Index]);
+				while (!RowArgs)
+				{
+					++Index;
+					RowArgs = MoveTemp(PopulateRowArgs[Index]);
+				}
+
+				PopulatePathDataTableRow(MoveTemp(RowArgs), Database, InRowHandle);
+				++Index;
+			});
+	}
+}
+ 
+void FTedsAssetData::OnPathsRemoved(TConstArrayView<FStringView> InPathsRemoved)
+{
+	TRACE_CPUPROFILER_EVENT_SCOPE(FTedsAssetData::OnPathsRemoved);
+
+#if TRACK_TEDSASSETDATA_MEMORY
+	LLM_SCOPE_BYNAME(TEXT("FTedsAssetData"))
+#endif
+
+	for (const FStringView Path : InPathsRemoved)
+	{
+		const TypedElementDataStorage::IndexHash PathHash = TypedElementDataStorage::GenerateIndexHash(FName(Path));
+		Database.RemoveRow(Database.FindIndexedRow(PathHash));
+		Database.RemoveIndex(PathHash);
+	}
+}
+}
+
+#undef TRACK_TEDSASSETDATA_MEMORY
\ No newline at end of file
diff --git a/Engine/Plugins/Experimental/TypedElementsDataStorage/Source/TypedElementsDataStorageAssetData/Private/TedsAssetData.h b/Engine/Plugins/Experimental/TypedElementsDataStorage/Source/TypedElementsDataStorageAssetData/Private/TedsAssetData.h
new file mode 100644
--- /dev/null
+++ b/Engine/Plugins/Experimental/TypedElementsDataStorage/Source/TypedElementsDataStorageAssetData/Private/TedsAssetData.h
@@ -0,0 +1,49 @@
+// Copyright Epic Games, Inc. All Rights Reserved.
+
+#pragma once
+
+#include "Containers/ContainersFwd.h"
+#include "Containers/StringFwd.h"
+#include "Elements/Common/TypedElementHandles.h"
+
+struct FAssetData;
+
+class ITypedElementDataStorageInterface;
+
+namespace UE::TypedElementsDataStorageAssetData::Private
+{
+
+/**
+ * Manage the registration and life cycle of the row related representing the data from the asset registry into TEDS.
+ */
+class FTedsAssetData
+{
+public:
+	FTedsAssetData(ITypedElementDataStorageInterface& InDatabase);
+
+	~FTedsAssetData();
+
+	void ProcessAllEvents();
+
+private:
+	void OnAssetsAdded(TConstArrayView<FAssetData> InAssetsAdded);
+	void OnAssetsRemoved(TConstArrayView<FAssetData> InAssetsRemoved);
+	void OnAssetsUpdated(TConstArrayView<FAssetData> InAssetsUpdated);
+	void OnAssetsUpdatedOnDisk(TConstArrayView<FAssetData> InAssetsUpdated);
+
+	void OnAssetRenamed(const FAssetData& InAsset, const FString& InOldObjectPath);
+
+	void OnPathsAdded(TConstArrayView<FStringView> InPathsAdded);
+	void OnPathsRemoved(TConstArrayView<FStringView> InPathsRemoved);
+
+	ITypedElementDataStorageInterface& Database;
+	TypedElementDataStorage::TableHandle PathsTable = TypedElementDataStorage::InvalidTableHandle;
+	TypedElementDataStorage::TableHandle AssetsDataTable = TypedElementDataStorage::InvalidTableHandle;
+
+	TypedElementDataStorage::QueryHandle UpdateAssetsInPath = TypedElementDataStorage::InvalidQueryHandle;
+	TypedElementDataStorage::QueryHandle ResolveMissingAssetInPathQuery = TypedElementDataStorage::InvalidQueryHandle;
+	TypedElementDataStorage::QueryHandle UpdateParentToChildrenAssetPathQuery = TypedElementDataStorage::InvalidQueryHandle;
+	TypedElementDataStorage::QueryHandle ResolveMissingParentPathQuery = TypedElementDataStorage::InvalidQueryHandle;
+};
+
+}
diff --git a/Engine/Plugins/Experimental/TypedElementsDataStorage/Source/TypedElementsDataStorageAssetData/Private/Tests/TedsAssetDataTest.cpp b/Engine/Plugins/Experimental/TypedElementsDataStorage/Source/TypedElementsDataStorageAssetData/Private/Tests/TedsAssetDataTest.cpp
new file mode 100644
--- /dev/null
+++ b/Engine/Plugins/Experimental/TypedElementsDataStorage/Source/TypedElementsDataStorageAssetData/Private/Tests/TedsAssetDataTest.cpp
@@ -0,0 +1,151 @@
+// Copyright Epic Games, Inc. All Rights Reserved.
+
+#include "Misc/AutomationTest.h"
+
+#include "AssetRegistry/AssetData.h"
+#include "AssetRegistry/IAssetRegistry.h"
+#include "Containers/ChunkedArray.h"
+#include "Containers/Set.h"
+#include "Containers/UnrealString.h"
+#include "Elements/Common/TypedElementHandles.h"
+#include "Elements/Framework/TypedElementIndexHasher.h"
+#include "Elements/Framework/TypedElementRegistry.h"
+#include "Elements/Interfaces/TypedElementDataStorageInterface.h"
+#include "Misc/ScopeExit.h"
+#include "TedsAssetDataColumns.h"
+#include "TypedElementsDataStorageAssetDataModule.h"
+#include "UObject/NameTypes.h"
+
+IMPLEMENT_SIMPLE_AUTOMATION_TEST(FTedsAssetDataTest, "TEDS.AssetRegistry.ValidateState", EAutomationTestFlags::EditorContext | EAutomationTestFlags::EngineFilter | EAutomationTestFlags::HighPriority)
+ 
+bool FTedsAssetDataTest::RunTest(const FString& Parameters)
+{
+	UE::TypedElementsDataStorageAssetData::FTypedElementsDataStorageAssetDataModule& TedsAssetRegistry = UE::TypedElementsDataStorageAssetData::FTypedElementsDataStorageAssetDataModule::GetChecked();
+
+	const bool bIsEnabled = TedsAssetRegistry.IsTedsAssetRegistryStorageEnabled();
+	ON_SCOPE_EXIT
+	{
+		if (!bIsEnabled)
+		{
+			TedsAssetRegistry.DisableTedsAssetRegistryStorage();
+		}
+	};
+
+	TedsAssetRegistry.EnableTedsAssetRegistryStorage();
+
+	TMap<FName, int32> AssetRegistryPathsAndAssetCount;
+
+	IAssetRegistry& AssetRegistry = IAssetRegistry::GetChecked();
+	AssetRegistry.WaitForCompletion();
+	 
+	AssetRegistry.EnumerateAllCachedPaths([&AssetRegistryPathsAndAssetCount](FName InPath)
+		{
+			AssetRegistryPathsAndAssetCount.Add(InPath, 0);
+			return true;
+		});
+
+	TChunkedArray<FAssetData> AssetRegistryAssetsData;
+	AssetRegistry.EnumerateAllAssets([&AssetRegistryAssetsData, &AssetRegistryPathsAndAssetCount](const FAssetData& InAssetData)
+		{
+			AssetRegistryAssetsData.AddElement(InAssetData);
+
+			if (int32* Count = AssetRegistryPathsAndAssetCount.Find(InAssetData.PackagePath))
+			{
+				//ensure(!InAssetData.OptionalOuterPath.IsNull());
+				++(*Count);
+			}
+
+			return true;
+		});
+
+	TedsAssetRegistry.ProcessDependentEvents();
+
+
+	// Do a sanity check that the data from the asset registry exist in Teds
+	const ITypedElementDataStorageInterface* Database = UTypedElementRegistry::GetInstance()->GetDataStorage();
+	
+	for (const TPair<FName, int32>& PathAndAssetCount : AssetRegistryPathsAndAssetCount)
+	{
+		TypedElementDataStorage::RowHandle Row = Database->FindIndexedRow(TypedElementDataStorage::GenerateIndexHash(PathAndAssetCount.Key));
+		bool bHadError = false;
+		bHadError |= !TestTrue(TEXT("Asset registry folder/path is indexed in TEDS"), Row != TypedElementDataStorage::InvalidRowHandle);
+
+		FName NameStoredInTeds;
+		if (const FAssetPathColumn_Experimental* AssetPathColumn = Database->GetColumn<FAssetPathColumn_Experimental>(Row))
+		{
+			NameStoredInTeds = AssetPathColumn->Path;
+		}
+
+		bHadError |= !TestEqual(TEXT("Path Stored In TEDS Match Path From Asset Registry"), NameStoredInTeds, PathAndAssetCount.Key);
+
+		if (const FAssetsInPathColumn_Experimental* AssetsInPathColumn = Database->GetColumn<FAssetsInPathColumn_Experimental>(Row))
+		{
+			bHadError |= !TestEqual(TEXT("Teds should have the same assets amount in a path as the asset registry has."), AssetsInPathColumn->AssetsRow.Num(), PathAndAssetCount.Value);
+		}
+
+		if (bHadError)
+		{
+			TArray<FStringFormatArg> FormatArgs;
+			FormatArgs.Emplace(PathAndAssetCount.Key.ToString());
+			FString Message = FString::Format(TEXT("Error was found at the following path:\"{0}\""), FormatArgs);
+			AddError(Message);
+
+			// Stop at the first error so that test don't take top long if something broke
+			break;
+		}
+	}
+
+	FNameBuilder Builder;
+	for (const FAssetData& AssetData : AssetRegistryAssetsData)
+	{
+		AssetData.PackagePath.ToString(Builder);
+		Builder.AppendChar(TEXT('/'));
+
+		// Verse package don't emit asset added event so we will ignore those for the test to avoid false errors
+		if (!FPackageName::IsVersePackage(Builder))
+		{ 
+			TypedElementDataStorage::RowHandle Row = Database->FindIndexedRow(TypedElementDataStorage::GenerateIndexHash(AssetData.GetSoftObjectPath()));
+
+			bool bHadError = false;
+
+			bHadError |= !TestTrue(TEXT("Asset registry asset path is indexed in TEDS"), Row != TypedElementDataStorage::InvalidRowHandle);
+
+			bHadError |= !TestNotNull(TEXT("Teds doesn't have a asset column for an asset of the asset registry"), Database->GetColumn<FAssetDataColumn_Experimental>(Row));
+
+			if (!Database->GetColumn<FUnresolvedAssetsInPathColumn_Experimental>(Row))
+			{
+				TypedElementDataStorage::RowHandle PathRow = Database->FindIndexedRow(TypedElementDataStorage::GenerateIndexHash(AssetData.PackagePath));
+			
+				const FAssetsInPathColumn_Experimental* AssetsInPathColumn = Database->GetColumn<FAssetsInPathColumn_Experimental>(PathRow);
+
+				const bool bPassNullTest = TestNotNull(TEXT("Asset in path column does exist for the package path of the asset"), AssetsInPathColumn);
+				if (bPassNullTest)
+				{
+					bHadError |= !TestTrue(TEXT("The asset is stored properly in the assets in path column"), AssetsInPathColumn->AssetsRow.Contains(Row));
+				}
+				else
+				{
+					bHadError = true;
+				}
+			}
+
+			if (bHadError)
+			{
+				TArray<FStringFormatArg> FormatArgs;
+				FormatArgs.Emplace(AssetData.GetSoftObjectPath().ToString());
+				FString Message = FString::Format(TEXT("Error was found for this asset:\"{0}\""), FormatArgs);
+				AddError(Message);
+
+				AssetData.PrintAssetData();
+
+				// Stop at the first error so that test don't take top long if something broke
+				break;
+			}
+		}
+	}
+
+	return true;
+}
+
+
+
diff --git a/Engine/Plugins/Experimental/TypedElementsDataStorage/Source/TypedElementsDataStorageAssetData/Private/TypedElementsDataStorageAssetDataModule.cpp b/Engine/Plugins/Experimental/TypedElementsDataStorage/Source/TypedElementsDataStorageAssetData/Private/TypedElementsDataStorageAssetDataModule.cpp
new file mode 100644
--- /dev/null
+++ b/Engine/Plugins/Experimental/TypedElementsDataStorage/Source/TypedElementsDataStorageAssetData/Private/TypedElementsDataStorageAssetDataModule.cpp
@@ -0,0 +1,113 @@
+// Copyright Epic Games, Inc. All Rights Reserved.
+
+#include "TypedElementsDataStorageAssetDataModule.h"
+
+#include "Elements/Framework/TypedElementRegistry.h"
+#include "HAL/IConsoleManager.h"
+#include "Modules/ModuleManager.h"
+#include "TedsAssetData.h"
+
+IMPLEMENT_MODULE(UE::TypedElementsDataStorageAssetData::FTypedElementsDataStorageAssetDataModule, TypedElementsDataStorageAssetData);
+
+namespace UE::TypedElementsDataStorageAssetData
+{
+
+namespace Private
+{
+TAutoConsoleVariable<bool> CVarTEDSAssetDataStorage(TEXT("Teds.AssetDataStorage"), false, TEXT("When true we will activate a wrapper that store the a copy of the asset data including the in memory change from the asset registry into TEDS.")
+	, FConsoleVariableDelegate::CreateLambda([](IConsoleVariable* Variable)
+	{
+		const bool bIsEnabled = Variable->GetBool();
+		FTypedElementsDataStorageAssetDataModule& Module = FTypedElementsDataStorageAssetDataModule::GetChecked();
+
+		if (bIsEnabled)
+		{
+			Module.EnableTedsAssetRegistryStorage();
+		}
+		else
+		{
+			Module.DisableTedsAssetRegistryStorage();
+		}
+	}));
+}
+
+void FTypedElementsDataStorageAssetDataModule::StartupModule()
+{
+	if (Private::CVarTEDSAssetDataStorage.GetValueOnGameThread())
+	{
+		EnableTedsAssetRegistryStorage();
+	}
+}
+
+void FTypedElementsDataStorageAssetDataModule::ShutdownModule()
+{
+	if (UTypedElementRegistry* TypedElementRegistry = UTypedElementRegistry::GetInstance())
+	{
+		TypedElementRegistry->OnDataStorageInterfacesSet().RemoveAll(this);
+	}
+}
+
+FTypedElementsDataStorageAssetDataModule* FTypedElementsDataStorageAssetDataModule::Get()
+{
+	return FModuleManager::Get().LoadModulePtr<FTypedElementsDataStorageAssetDataModule>(TEXT("TypedElementsDataStorageAssetData"));
+}
+
+FTypedElementsDataStorageAssetDataModule& FTypedElementsDataStorageAssetDataModule::GetChecked()
+{
+	return FModuleManager::Get().LoadModuleChecked<FTypedElementsDataStorageAssetDataModule>(TEXT("TypedElementsDataStorageAssetData"));
+}
+
+void FTypedElementsDataStorageAssetDataModule::EnableTedsAssetRegistryStorage()
+{
+	if (!AssetRegistryStorage)
+	{
+		FModuleManager::Get().LoadModuleChecked(TEXT("TypedElementFramework"));
+		UTypedElementRegistry* TypedElementRegistry = UTypedElementRegistry::GetInstance();
+		if (TypedElementRegistry->GetMutableDataStorage())
+		{
+			InitAssetRegistryStorage();
+		}
+		else
+		{
+			TypedElementRegistry->OnDataStorageInterfacesSet().AddRaw(this, &FTypedElementsDataStorageAssetDataModule::InitAssetRegistryStorage);
+		}
+
+		if (!Private::CVarTEDSAssetDataStorage.GetValueOnGameThread())
+		{
+			Private::CVarTEDSAssetDataStorage.AsVariable()->Set(true);
+		}
+	}
+}
+
+void FTypedElementsDataStorageAssetDataModule::DisableTedsAssetRegistryStorage()
+{
+	if (AssetRegistryStorage)
+	{
+		AssetRegistryStorage.Reset();
+
+		if (Private::CVarTEDSAssetDataStorage.GetValueOnGameThread())
+		{
+			Private::CVarTEDSAssetDataStorage.AsVariable()->Set(false);
+		}
+	}
+}
+
+bool FTypedElementsDataStorageAssetDataModule::IsTedsAssetRegistryStorageEnabled() const
+{
+	return AssetRegistryStorage.IsValid();
+}
+
+void FTypedElementsDataStorageAssetDataModule::ProcessDependentEvents()
+{
+	if (Private::FTedsAssetData* Storage = AssetRegistryStorage.Get())
+	{
+		Storage->ProcessAllEvents();
+	}
+}
+
+void FTypedElementsDataStorageAssetDataModule::InitAssetRegistryStorage()
+{
+	AssetRegistryStorage = MakeUnique<Private::FTedsAssetData>(*UTypedElementRegistry::GetInstance()->GetMutableDataStorage());
+}
+
+}
\ No newline at end of file
diff --git a/Engine/Plugins/Experimental/TypedElementsDataStorage/Source/TypedElementsDataStorageAssetData/Private/TypedElementsDataStorageAssetDataModule.h b/Engine/Plugins/Experimental/TypedElementsDataStorage/Source/TypedElementsDataStorageAssetData/Private/TypedElementsDataStorageAssetDataModule.h
new file mode 100644
--- /dev/null
+++ b/Engine/Plugins/Experimental/TypedElementsDataStorage/Source/TypedElementsDataStorageAssetData/Private/TypedElementsDataStorageAssetDataModule.h
@@ -0,0 +1,42 @@
+// Copyright Epic Games, Inc. All Rights Reserved.
+
+#pragma once
+
+#include "Modules/ModuleInterface.h"
+#include "Templates/UniquePtr.h"
+
+namespace UE::TypedElementsDataStorageAssetData
+{ 
+
+namespace Private
+{
+	class FTedsAssetData;
+}
+
+class FTypedElementsDataStorageAssetDataModule : public IModuleInterface
+{
+public:
+	virtual void StartupModule() override;
+	virtual void ShutdownModule() override;
+
+	static FTypedElementsDataStorageAssetDataModule* Get();
+	static FTypedElementsDataStorageAssetDataModule& GetChecked();
+
+	void EnableTedsAssetRegistryStorage();
+	void DisableTedsAssetRegistryStorage();
+	bool IsTedsAssetRegistryStorageEnabled() const;
+
+
+	/**
+	 * Process now any pending event that might make the Teds database out of sync with the asset registry.
+	 * Note: This isn't needed when using the editor so it should be only called by the automation scripts that need it to avoid creating some unneeded stalls.
+	 */
+	void ProcessDependentEvents();
+
+private:
+	void InitAssetRegistryStorage();
+
+	TUniquePtr<Private::FTedsAssetData> AssetRegistryStorage;
+};
+
+}
diff --git a/Engine/Plugins/Experimental/TypedElementsDataStorage/Source/TypedElementsDataStorageAssetData/Public/TedsAssetDataColumns.h b/Engine/Plugins/Experimental/TypedElementsDataStorage/Source/TypedElementsDataStorageAssetData/Public/TedsAssetDataColumns.h
new file mode 100644
--- /dev/null
+++ b/Engine/Plugins/Experimental/TypedElementsDataStorage/Source/TypedElementsDataStorageAssetData/Public/TedsAssetDataColumns.h
@@ -0,0 +1,67 @@
+// Copyright Epic Games, Inc. All Rights Reserved.
+
+#pragma once
+
+#include "AssetRegistry/AssetData.h"
+#include "Elements/Common/TypedElementCommonTypes.h"
+#include "Elements/Interfaces/TypedElementQueryStorageInterfaces.h"
+
+#include "TedsAssetDataColumns.generated.h"
+
+USTRUCT()
+struct FAssetPathColumn_Experimental : public FTypedElementDataStorageColumn
+{
+	GENERATED_BODY()
+	
+	UPROPERTY()
+	FName Path;
+};
+
+USTRUCT()
+struct FParentAssetPathColumn_Experimental : public FTypedElementDataStorageColumn
+{
+	GENERATED_BODY()
+	
+	TypedElementDataStorage::RowHandle ParentRow;
+};
+
+USTRUCT()
+struct FChildrenAssetPathColumn_Experimental : public FTypedElementDataStorageColumn
+{
+	GENERATED_BODY()
+	
+	TSet<TypedElementDataStorage::RowHandle> ChildrenRows;
+};
+
+USTRUCT()
+struct FAssetsInPathColumn_Experimental : public FTypedElementDataStorageColumn
+{
+	GENERATED_BODY()
+	
+	TSet<TypedElementDataStorage::RowHandle> AssetsRow;
+};
+
+USTRUCT()
+struct FAssetDataColumn_Experimental : public FTypedElementDataStorageColumn
+{
+	GENERATED_BODY()
+	
+	UPROPERTY()
+	FAssetData AssetData;
+};
+
+USTRUCT()
+struct FUnresolvedParentAssetPathColumn_Experimental : public FTypedElementDataStorageColumn
+{
+	GENERATED_BODY()
+	
+	TypedElementDataStorage::IndexHash Hash;
+};
+
+USTRUCT()
+struct FUnresolvedAssetsInPathColumn_Experimental : public FTypedElementDataStorageColumn
+{
+	GENERATED_BODY()
+	
+	TypedElementDataStorage::IndexHash Hash;
+};
diff --git a/Engine/Plugins/Experimental/TypedElementsDataStorage/Source/TypedElementsDataStorageAssetData/TypedElementsDataStorageAssetData.build.cs b/Engine/Plugins/Experimental/TypedElementsDataStorage/Source/TypedElementsDataStorageAssetData/TypedElementsDataStorageAssetData.build.cs
new file mode 100644
--- /dev/null
+++ b/Engine/Plugins/Experimental/TypedElementsDataStorage/Source/TypedElementsDataStorageAssetData/TypedElementsDataStorageAssetData.build.cs
@@ -0,0 +1,29 @@
+// Copyright Epic Games, Inc. All Rights Reserved.
+
+using UnrealBuildTool;
+
+// Experimental test module. please refrain from depending on it until this warning is removed
+public class TypedElementsDataStorageAssetData : ModuleRules
+{
+	public TypedElementsDataStorageAssetData(ReadOnlyTargetRules Target) : base(Target)
+	{
+		PCHUsage = ModuleRules.PCHUsageMode.UseExplicitOrSharedPCHs;
+
+		if (Target.bBuildEditor)
+		{
+			PublicIncludePaths.AddRange(new string[] { });
+			PrivateIncludePaths.AddRange(new string[] { });
+
+			PrivateDependencyModuleNames.AddRange(
+				new string[]
+				{
+					"AssetRegistry",
+					"Core",
+					"CoreUObject",
+					"Engine",
+					"TypedElementFramework",
+				});
+
+		}
+	}
+}
diff --git a/Engine/Plugins/Experimental/TypedElementsDataStorage/TypedElementsDataStorage.uplugin b/Engine/Plugins/Experimental/TypedElementsDataStorage/TypedElementsDataStorage.uplugin
--- a/Engine/Plugins/Experimental/TypedElementsDataStorage/TypedElementsDataStorage.uplugin
+++ b/Engine/Plugins/Experimental/TypedElementsDataStorage/TypedElementsDataStorage.uplugin
@@ -26,6 +26,10 @@
 		{
 			"Name": "TypedElementsDataStorageRevisionControl",
 			"Type": "Editor"
+		},
+		{
+			"Name": "TypedElementsDataStorageAssetData",
+			"Type":  "Editor"
 		}
 	],
 	"Plugins": [
