Added a sort name to the TypedElementDataStorageAssetData module.

diff --git a/Engine/Plugins/Experimental/TypedElementsDataStorage/Source/TypedElementsDataStorageAssetData/TypedElementsDataStorageAssetData.build.cs b/Engine/Plugins/Experimental/TypedElementsDataStorage/Source/TypedElementsDataStorageAssetData/TypedElementsDataStorageAssetData.build.cs
--- a/Engine/Plugins/Experimental/TypedElementsDataStorage/Source/TypedElementsDataStorageAssetData/TypedElementsDataStorageAssetData.build.cs
+++ b/Engine/Plugins/Experimental/TypedElementsDataStorage/Source/TypedElementsDataStorageAssetData/TypedElementsDataStorageAssetData.build.cs
@@ -7,6 +7,8 @@ public class TypedElementsDataStorageAssetData : ModuleRules
 {
 	public TypedElementsDataStorageAssetData(ReadOnlyTargetRules Target) : base(Target)
 	{
+		ShortName = "TEDSAssetD";
+
 		PCHUsage = ModuleRules.PCHUsageMode.UseExplicitOrSharedPCHs;
 
 		if (Target.bBuildEditor)
