TEDS Asset Data: Move asset columns and processors to the TypedElementsDataStorageAssetData module

diff --git a/Engine/Plugins/Experimental/TypedElementsDataStorage/Source/TypedElementsDataStorage/Private/AssetDebugCommands.cpp b/Engine/Plugins/Experimental/TypedElementsDataStorage/Source/TypedElementsDataStorage/Private/AssetDebugCommands.cpp
deleted file mode 100644
diff --git a/Engine/Plugins/Experimental/TypedElementsDataStorage/Source/TypedElementsDataStorage/Private/AssetProcessors.cpp b/Engine/Plugins/Experimental/TypedElementsDataStorage/Source/TypedElementsDataStorage/Private/AssetProcessors.cpp
deleted file mode 100644
diff --git a/Engine/Plugins/Experimental/TypedElementsDataStorage/Source/TypedElementsDataStorage/Private/AssetProcessors.h b/Engine/Plugins/Experimental/TypedElementsDataStorage/Source/TypedElementsDataStorage/Private/AssetProcessors.h
deleted file mode 100644
diff --git a/Engine/Plugins/Experimental/TypedElementsDataStorage/Source/TypedElementsDataStorage/Public/TypedElementAssetColumns.h b/Engine/Plugins/Experimental/TypedElementsDataStorage/Source/TypedElementsDataStorage/Public/TypedElementAssetColumns.h
deleted file mode 100644
diff --git a/Engine/Plugins/Experimental/TypedElementsDataStorage/Source/TypedElementsDataStorageAssetData/Private/Debug/AssetProcessors.cpp b/Engine/Plugins/Experimental/TypedElementsDataStorage/Source/TypedElementsDataStorageAssetData/Private/Debug/AssetProcessors.cpp
new file mode 100644
--- /dev/null
+++ b/Engine/Plugins/Experimental/TypedElementsDataStorage/Source/TypedElementsDataStorageAssetData/Private/Debug/AssetProcessors.cpp
@@ -0,0 +1,16 @@
+// Copyright Epic Games, Inc. All Rights Reserved.
+
+#include "AssetProcessors.h"
+
+#include "Elements/Columns/TypedElementLabelColumns.h"
+#include "Elements/Columns/TypedElementTypeInfoColumns.h"
+#include "TedsAssetDataColumns.h"
+
+void UTypedElementAssetFactory::RegisterTables(ITypedElementDataStorageInterface& DataStorage) 
+{
+	DataStorage.RegisterTable(
+		TTypedElementColumnTypeList<
+			FTypedElementLabelColumn, FTypedElementClassTypeInfoColumn, FAssetPathColumn_Experimental,
+			FAssetTag, FDiskSizeColumn, FVersePathColumn>(),
+		FName("Editor_PlaceholderAssetTable"));
+}
diff --git a/Engine/Plugins/Experimental/TypedElementsDataStorage/Source/TypedElementsDataStorageAssetData/Private/Debug/AssetProcessors.h b/Engine/Plugins/Experimental/TypedElementsDataStorage/Source/TypedElementsDataStorageAssetData/Private/Debug/AssetProcessors.h
new file mode 100644
--- /dev/null
+++ b/Engine/Plugins/Experimental/TypedElementsDataStorage/Source/TypedElementsDataStorageAssetData/Private/Debug/AssetProcessors.h
@@ -0,0 +1,22 @@
+// Copyright Epic Games, Inc. All Rights Reserved.
+
+#pragma once
+
+#include "Elements/Interfaces/TypedElementDataStorageFactory.h"
+#include "Elements/Interfaces/TypedElementQueryStorageInterfaces.h"
+#include "UObject/ObjectMacros.h"
+
+#include "AssetProcessors.generated.h"
+
+UCLASS()
+class UTypedElementAssetFactory : public UTypedElementDataStorageFactory
+{
+	GENERATED_BODY()
+
+public:
+	~UTypedElementAssetFactory() override = default;
+
+	void RegisterTables(ITypedElementDataStorageInterface& DataStorage) override;
+
+private:
+};
diff --git a/Engine/Plugins/Experimental/TypedElementsDataStorage/Source/TypedElementsDataStorageAssetData/Private/Debug/TedsAssetDataDebugCommands.cpp b/Engine/Plugins/Experimental/TypedElementsDataStorage/Source/TypedElementsDataStorageAssetData/Private/Debug/TedsAssetDataDebugCommands.cpp
--- a/Engine/Plugins/Experimental/TypedElementsDataStorage/Source/TypedElementsDataStorageAssetData/Private/Debug/TedsAssetDataDebugCommands.cpp
+++ b/Engine/Plugins/Experimental/TypedElementsDataStorage/Source/TypedElementsDataStorageAssetData/Private/Debug/TedsAssetDataDebugCommands.cpp
@@ -9,136 +9,274 @@
 #include "HAL/IConsoleManager.h"
 #include "TedsAssetDataColumns.h"
 #include "UObject/NameTypes.h"
+#include "Engine/StaticMesh.h"
+#include "Materials/Material.h"
+#include "Engine/Level.h"
+#include "Engine/Blueprint.h"
+#include "Engine/Texture.h"
+#include "Elements/Columns/TypedElementLabelColumns.h"
+#include "Elements/Columns/TypedElementTypeInfoColumns.h"
+#include "Elements/Framework/TypedElementQueryBuilder.h"
+#include "HAL/IConsoleManager.h"
+#include "Math/UnrealMathUtility.h"
 
 
 DEFINE_LOG_CATEGORY_STATIC(LogTEDSAssetRegistry, Log, All)
 
-static FAutoConsoleCommand CCMDTestFolderRowData(
-	TEXT("TEDS.Debug.ShowDataOfAssetFolder"),
-	TEXT("Print some debug information on the specified path."),
-	FConsoleCommandWithArgsDelegate::CreateLambda([](const TArray<FString>& InArgs)
-		{
-			const ITypedElementDataStorageInterface* Database = UTypedElementRegistry::GetInstance()->GetDataStorage();
-			IAssetRegistry& AssetRegistry = IAssetRegistry::GetChecked();
-
-			for (const FString& Path : InArgs)
+namespace UE::TEDSAssetDataDebugCommands::Local
+{
+	static FAutoConsoleCommand CCMDTestFolderRowData(
+		TEXT("TEDS.Debug.ShowDataOfAssetFolder"),
+		TEXT("Print some debug information on the specified path."),
+		FConsoleCommandWithArgsDelegate::CreateLambda([](const TArray<FString>& InArgs)
 			{
-				FName PathAsName(*Path);
-				TypedElementDataStorage::RowHandle RowHandle = Database->FindIndexedRow(TypedElementDataStorage::GenerateIndexHash(PathAsName));
+				const ITypedElementDataStorageInterface* Database = UTypedElementRegistry::GetInstance()->GetDataStorage();
+				IAssetRegistry& AssetRegistry = IAssetRegistry::GetChecked();
 
-				if (Database->IsRowAssigned(RowHandle))
+				for (const FString& Path : InArgs)
 				{
-					UE_LOG(LogTEDSAssetRegistry, Display, TEXT("The path isn't indexed."));
-					return;
-				}
+					FName PathAsName(*Path);
+					TypedElementDataStorage::RowHandle RowHandle = Database->FindIndexedRow(TypedElementDataStorage::GenerateIndexHash(PathAsName));
 
+					if (Database->IsRowAssigned(RowHandle))
+					{
+						UE_LOG(LogTEDSAssetRegistry, Display, TEXT("The path isn't indexed."));
+						return;
+					}
 
-				UE_LOG(LogTEDSAssetRegistry, Display, TEXT("Found some information for the path (%s) in the database."), *Path);
 
-				if (const FAssetPathColumn_Experimental* AssetPath = Database->GetColumn<FAssetPathColumn_Experimental>(RowHandle))
-				{
-					UE_LOG(LogTEDSAssetRegistry, Display, TEXT("Path stored in the database as (%s)."), *AssetPath->Path.ToString());
-				}
+					UE_LOG(LogTEDSAssetRegistry, Display, TEXT("Found some information for the path (%s) in the database."), *Path);
 
-				if (const FParentAssetPathColumn_Experimental* ParentAssetPath = Database->GetColumn<FParentAssetPathColumn_Experimental>(RowHandle))
-				{
-					if (const FAssetPathColumn_Experimental* AssetPath = Database->GetColumn<FAssetPathColumn_Experimental>(ParentAssetPath->ParentRow))
+					if (const FAssetPathColumn_Experimental* AssetPath = Database->GetColumn<FAssetPathColumn_Experimental>(RowHandle))
 					{
-						UE_LOG(LogTEDSAssetRegistry, Display, TEXT("	Parent Path: %s"), *AssetPath->Path.ToString());
+						UE_LOG(LogTEDSAssetRegistry, Display, TEXT("Path stored in the database as (%s)."), *AssetPath->Path.ToString());
 					}
-				}
 
+					if (const FParentAssetPathColumn_Experimental* ParentAssetPath = Database->GetColumn<FParentAssetPathColumn_Experimental>(RowHandle))
+					{
+						if (const FAssetPathColumn_Experimental* AssetPath = Database->GetColumn<FAssetPathColumn_Experimental>(ParentAssetPath->ParentRow))
+						{
+							UE_LOG(LogTEDSAssetRegistry, Display, TEXT("	Parent Path: %s"), *AssetPath->Path.ToString());
+						}
+					}
 
-				if (const FChildrenAssetPathColumn_Experimental* ChildrenPath = Database->GetColumn<FChildrenAssetPathColumn_Experimental>(RowHandle))
-				{
-					UE_LOG(LogTEDSAssetRegistry, Display, TEXT("	Path as %i children"), ChildrenPath->ChildrenRows.Num());
 
-					for (TypedElementDataStorage::RowHandle Row : ChildrenPath->ChildrenRows)
+					if (const FChildrenAssetPathColumn_Experimental* ChildrenPath = Database->GetColumn<FChildrenAssetPathColumn_Experimental>(RowHandle))
 					{
-						if (const FAssetPathColumn_Experimental* AssetPath = Database->GetColumn<FAssetPathColumn_Experimental>(Row))
+						UE_LOG(LogTEDSAssetRegistry, Display, TEXT("	Path as %i children"), ChildrenPath->ChildrenRows.Num());
+
+						for (TypedElementDataStorage::RowHandle Row : ChildrenPath->ChildrenRows)
 						{
-							UE_LOG(LogTEDSAssetRegistry, Display, TEXT("		Children Path: %s"), *AssetPath->Path.ToString());
+							if (const FAssetPathColumn_Experimental* AssetPath = Database->GetColumn<FAssetPathColumn_Experimental>(Row))
+							{
+								UE_LOG(LogTEDSAssetRegistry, Display, TEXT("		Children Path: %s"), *AssetPath->Path.ToString());
+							}
 						}
 					}
-				}
 
 
-				if (const FAssetsInPathColumn_Experimental* AssetInPath = Database->GetColumn<FAssetsInPathColumn_Experimental>(RowHandle))
-				{
-					UE_LOG(LogTEDSAssetRegistry, Display, TEXT("	Asset in Paths"));
-
-					for (TypedElementDataStorage::RowHandle AssetRow : AssetInPath->AssetsRow)
+					if (const FAssetsInPathColumn_Experimental* AssetInPath = Database->GetColumn<FAssetsInPathColumn_Experimental>(RowHandle))
 					{
-						if (const FAssetDataColumn_Experimental* AssetData = Database->GetColumn<FAssetDataColumn_Experimental>(AssetRow))
+						UE_LOG(LogTEDSAssetRegistry, Display, TEXT("	Asset in Paths"));
+
+						for (TypedElementDataStorage::RowHandle AssetRow : AssetInPath->AssetsRow)
 						{
-							UE_LOG(LogTEDSAssetRegistry, Display, TEXT("		Asset Name: %s"), *AssetData->AssetData.AssetName.ToString());
+							if (const FAssetDataColumn_Experimental* AssetData = Database->GetColumn<FAssetDataColumn_Experimental>(AssetRow))
+							{
+								UE_LOG(LogTEDSAssetRegistry, Display, TEXT("		Asset Name: %s"), *AssetData->AssetData.AssetName.ToString());
+							}
+							else
+							{
+								UE_LOG(LogTEDSAssetRegistry, Display, TEXT("		Asset Row pointed to stale asset."));
+							}
 						}
-						else
+					}
+
+
+					// Check for asset that haven't processed their path asset in path column yet
+					TArray<FAssetData> Assets;
+					AssetRegistry.GetAssetsByPath(PathAsName, Assets);
+					for (const FAssetData& Asset : Assets)
+					{
+						TypedElementDataStorage::RowHandle AssetRow = Database->FindIndexedRow(TypedElementDataStorage::GenerateIndexHash(Asset.GetSoftObjectPath()));
+
+						if (const FUnresolvedAssetsInPathColumn_Experimental* UnresolvedAssetsInPathColumn = Database->GetColumn<FUnresolvedAssetsInPathColumn_Experimental>(AssetRow))
 						{
-							UE_LOG(LogTEDSAssetRegistry, Display, TEXT("		Asset Row pointed to stale asset."));
+								UE_LOG(LogTEDSAssetRegistry, Display, TEXT("		Unresolved asset in path asset. Asset Name: %s"), *Asset.AssetName.ToString());
 						}
 					}
 				}
+			}));
 
 
-				// Check for asset that haven't processed their path asset in path column yet
-				TArray<FAssetData> Assets;
-				AssetRegistry.GetAssetsByPath(PathAsName, Assets);
-				for (const FAssetData& Asset : Assets)
+	static FAutoConsoleCommand CCMDTestFolderAssetRegistryData(
+		TEXT("TEDS.Debug.ShowAssetRegistryDataOfFolder"),
+		TEXT("Print some debug information on the specified path."),
+		FConsoleCommandWithArgsDelegate::CreateLambda([](const TArray<FString>& InArgs)
+			{
+				IAssetRegistry& AssetRegistry = IAssetRegistry::GetChecked();
+
+				FNameBuilder NameBuilder;
+				for (const FString& Path : InArgs)
 				{
-					TypedElementDataStorage::RowHandle AssetRow = Database->FindIndexedRow(TypedElementDataStorage::GenerateIndexHash(Asset.GetSoftObjectPath()));
+					UE_LOG(LogTEDSAssetRegistry, Display, TEXT("Displaying asset registry info on path (%s)"), *Path);
+					const FName PathAsName(*Path);
+					AssetRegistry.EnumerateSubPaths(PathAsName, [&NameBuilder](FName InPath)
+						{
+							InPath.AppendString(NameBuilder);
+							UE_LOG(LogTEDSAssetRegistry, Display, TEXT("	Children Path: %s"), *NameBuilder);
+							NameBuilder.Reset();
+							return true;
+						}, false);
 
-					if (const FUnresolvedAssetsInPathColumn_Experimental* UnresolvedAssetsInPathColumn = Database->GetColumn<FUnresolvedAssetsInPathColumn_Experimental>(AssetRow))
-					{
-							UE_LOG(LogTEDSAssetRegistry, Display, TEXT("		Unresolved asset in path asset. Asset Name: %s"), *Asset.AssetName.ToString());
+
+					TArray<FAssetData> Assets;
+					AssetRegistry.GetAssetsByPath(PathAsName, Assets);
+
+					if (!Assets.IsEmpty())
+					{ 
+						UE_LOG(LogTEDSAssetRegistry, Display, TEXT("	Asset in Path"));
+
+						for (const FAssetData& Asset : Assets)
+						{
+							Asset.AssetName.AppendString(NameBuilder);
+							UE_LOG(LogTEDSAssetRegistry, Display, TEXT("		Asset Name: %s"), *NameBuilder);
+							NameBuilder.Reset();
+
+							Asset.GetFullName(NameBuilder);
+							UE_LOG(LogTEDSAssetRegistry, Display, TEXT("		Asset Full Name: %s"), *NameBuilder);
+							NameBuilder.Reset();
+
+							Asset.PackageName.AppendString(NameBuilder);
+							UE_LOG(LogTEDSAssetRegistry, Display, TEXT("		Asset Reported Package Path: %s"), *NameBuilder);
+							NameBuilder.Reset();
+						}
 					}
 				}
+			}));
+	
+	// A small set of classes to randomly pick from
+	static TArray<UClass*> AssetClasses{ UStaticMesh::StaticClass(), UMaterial::StaticClass(),
+		ULevel::StaticClass(), UBlueprint::StaticClass(), UTexture::StaticClass() };
+
+	// Populate an asset row with random information
+	void PopulateRowWithRandomInfo(TypedElementDataStorage::RowHandle Row, ITypedElementDataStorageInterface* DataStorage)
+	{
+		// Don't modify any rows that aren't our placeholder assets
+		if (!DataStorage->HasColumns<FAssetTag>(Row))
+		{
+			return;
+		}
+
+		// Pick a random asset class from our list
+		UClass* AssetClass = AssetClasses[FMath::RandRange(0, AssetClasses.Num() - 1)];
+
+		// Add a label to the row
+		if (FTypedElementLabelColumn* LabelColumn = DataStorage->GetColumn<FTypedElementLabelColumn>(Row))
+		{
+			// We can have duplicate names but it doesn't really matter
+			LabelColumn->Label = AssetClass->GetFName().ToString();
+			LabelColumn->Label.Append(TEXT("_Placeholder_"));
+			LabelColumn->Label.Append(FString::FromInt(FMath::RandRange(0, 1000)));
+
+			if (FVersePathColumn* ClassTypeInfoColumn = DataStorage->GetColumn<FVersePathColumn>(Row))
+			{
+				FString TestVerseModule = TEXT("/UnrealEngine.com/Temporary/TEDS/");
+				TestVerseModule.Append(LabelColumn->Label);
+			
+				UE::Core::FVersePath::TryMake(ClassTypeInfoColumn->VersePath, TestVerseModule);
 			}
-		}));
 
+			if (FAssetPathColumn_Experimental* AssetPathColumn = DataStorage->GetColumn<FAssetPathColumn_Experimental>(Row))
+			{
+				FString TestAssetPath = TEXT("TestPath/TestDirectory/");
+				TestAssetPath.Append(LabelColumn->Label);
+				
+				AssetPathColumn->Path = FName(TestAssetPath);
+			}
+		}
+		
 
-static FAutoConsoleCommand CCMDTestFolderAssetRegistryData(
-	TEXT("TEDS.Debug.ShowAssetRegistryDataOfFolder"),
-	TEXT("Print some debug information on the specified path."),
-	FConsoleCommandWithArgsDelegate::CreateLambda([](const TArray<FString>& InArgs)
+		if (FTypedElementClassTypeInfoColumn* ClassTypeInfoColumn = DataStorage->GetColumn<FTypedElementClassTypeInfoColumn>(Row))
+		{
+			ClassTypeInfoColumn->TypeInfo = AssetClass;
+		}
+		
+		if (FDiskSizeColumn* ClassTypeInfoColumn = DataStorage->GetColumn<FDiskSizeColumn>(Row))
 		{
-			IAssetRegistry& AssetRegistry = IAssetRegistry::GetChecked();
+			ClassTypeInfoColumn->DiskSize = FMath::RandRange(1024, 32768);
+		}
 
-			FNameBuilder NameBuilder;
-			for (const FString& Path : InArgs)
+		// Randomly make this a public or a private asset
+		DataStorage->AddColumn(Row, FMath::RandRange(0, 1) == 0 ? FPrivateAssetTag::StaticStruct() : FPublicAssetTag::StaticStruct());
+	}
+	
+	FAutoConsoleCommand CreateDebugAssetRowsCommand(
+		TEXT("TEDS.Debug.CreateDebugAssetRows"),
+		TEXT("Create random asset rows. Args: (TEDS.Debug.CreateDebugAssetRows NumRows) Default 10 rows"),
+		FConsoleCommandWithArgsDelegate::CreateLambda(
+			[](const TArray<FString>& Args)
 			{
-				UE_LOG(LogTEDSAssetRegistry, Display, TEXT("Displaying asset registry info on path (%s)"), *Path);
-				const FName PathAsName(*Path);
-				AssetRegistry.EnumerateSubPaths(PathAsName, [&NameBuilder](FName InPath)
-					{
-						InPath.AppendString(NameBuilder);
-						UE_LOG(LogTEDSAssetRegistry, Display, TEXT("	Children Path: %s"), *NameBuilder);
-						NameBuilder.Reset();
-						return true;
-					}, false);
+				TRACE_CPUPROFILER_EVENT_SCOPE(TEDS.Debug.CreateDebugAssetRows);
 
+				using namespace TypedElementDataStorage;
+				using namespace TypedElementQueryBuilder;
 
-				TArray<FAssetData> Assets;
-				AssetRegistry.GetAssetsByPath(PathAsName, Assets);
+				ITypedElementDataStorageInterface* DataStorage = UTypedElementRegistry::GetInstance()->GetMutableDataStorage();
 
-				if (!Assets.IsEmpty())
-				{ 
-					UE_LOG(LogTEDSAssetRegistry, Display, TEXT("	Asset in Path"));
+				if (!DataStorage)
+				{
+					return;
+				}
 
-					for (const FAssetData& Asset : Assets)
-					{
-						Asset.AssetName.AppendString(NameBuilder);
-						UE_LOG(LogTEDSAssetRegistry, Display, TEXT("		Asset Name: %s"), *NameBuilder);
-						NameBuilder.Reset();
+				int RowCount = 10; // Create 10 rows by default
 
-						Asset.GetFullName(NameBuilder);
-						UE_LOG(LogTEDSAssetRegistry, Display, TEXT("		Asset Full Name: %s"), *NameBuilder);
-						NameBuilder.Reset();
+				if (Args.Num() == 1)
+				{
+					RowCount = FCString::Atoi(*Args[0]);
+				}
 
-						Asset.PackageName.AppendString(NameBuilder);
-						UE_LOG(LogTEDSAssetRegistry, Display, TEXT("		Asset Reported Package Path: %s"), *NameBuilder);
-						NameBuilder.Reset();
-					}
+				TableHandle Table = DataStorage->FindTable(FName("Editor_PlaceholderAssetTable"));
+
+				DataStorage->BatchAddRow(Table, RowCount, [DataStorage](RowHandle Row)
+				{
+					PopulateRowWithRandomInfo(Row, DataStorage);
+				});
+				
+			}
+		));
+
+	FAutoConsoleCommand RemoveAssetRowsCommand(
+		TEXT("TEDS.Debug.RemoveAssetRows"),
+		TEXT("Remove All Asset Rows"),
+		FConsoleCommandWithArgsDelegate::CreateLambda(
+			[](const TArray<FString>& Args)
+			{
+				using namespace TypedElementDataStorage;
+				using namespace TypedElementQueryBuilder;
+
+				ITypedElementDataStorageInterface* DataStorage = UTypedElementRegistry::GetInstance()->GetMutableDataStorage();
+
+				if (!DataStorage)
+				{
+					return;
+				}
+
+				// Select all rows with FAssetTag
+				static QueryHandle AssetQueryHandle = DataStorage->RegisterQuery(Select().Where().All<FAssetTag>().Compile());
+
+				TArray<RowHandle> Rows;
+				
+				DataStorage->RunQuery(AssetQueryHandle,
+					CreateDirectQueryCallbackBinding([&Rows](const ITypedElementDataStorageInterface::IDirectQueryContext& Context)
+					{
+						Rows.Append(Context.GetRowHandles());
+					}));
+
+				for(const RowHandle Row : Rows)
+				{
+					DataStorage->RemoveRow(Row);
 				}
 			}
-		}));
\ No newline at end of file
+		));
+}
diff --git a/Engine/Plugins/Experimental/TypedElementsDataStorage/Source/TypedElementsDataStorageAssetData/Public/TedsAssetDataColumns.h b/Engine/Plugins/Experimental/TypedElementsDataStorage/Source/TypedElementsDataStorageAssetData/Public/TedsAssetDataColumns.h
--- a/Engine/Plugins/Experimental/TypedElementsDataStorage/Source/TypedElementsDataStorageAssetData/Public/TedsAssetDataColumns.h
+++ b/Engine/Plugins/Experimental/TypedElementsDataStorage/Source/TypedElementsDataStorageAssetData/Public/TedsAssetDataColumns.h
@@ -5,6 +5,7 @@
 #include "AssetRegistry/AssetData.h"
 #include "Elements/Common/TypedElementCommonTypes.h"
 #include "Elements/Interfaces/TypedElementQueryStorageInterfaces.h"
+#include "Containers/VersePath.h"
 
 #include "TedsAssetDataColumns.generated.h"
 
@@ -65,3 +66,44 @@ struct FUnresolvedAssetsInPathColumn_Experimental : public FTypedElementDataStor
 	
 	TypedElementDataStorage::IndexHash Hash;
 };
+
+// Tag to identify assets
+USTRUCT(meta = (DisplayName = "Asset"))
+struct FAssetTag final : public FTypedElementDataStorageTag
+{
+	GENERATED_BODY()
+};
+
+// Tag to identify assets with private visibility
+USTRUCT(meta = (DisplayName = "Private Asset"))
+struct FPrivateAssetTag final : public FTypedElementDataStorageTag
+{
+	GENERATED_BODY()
+};
+
+// Tag to identify assets with public visibility
+USTRUCT(meta = (DisplayName = "Public Asset"))
+struct FPublicAssetTag final : public FTypedElementDataStorageTag
+{
+	GENERATED_BODY()
+};
+
+// Column to store the disk size of an asset
+USTRUCT(meta = (DisplayName = "Disk Size"))
+struct FDiskSizeColumn final : public FTypedElementDataStorageColumn
+{
+	GENERATED_BODY()
+
+	// Size on disk (in bytes)
+	UPROPERTY()
+	int64 DiskSize;
+};
+
+// Column to store the verse path of an asset
+USTRUCT(meta = (DisplayName = "Verse Path"))
+struct FVersePathColumn final : public FTypedElementDataStorageColumn
+{
+	GENERATED_BODY()
+	
+	UE::Core::FVersePath VersePath;
+};
\ No newline at end of file
