Disabled global TEDS lock for the moment.

There's an order-of-execution problem with TEDS and Mass that's causing the global thread lock to be held on too long. For now this has been removed for now, though this will not change the thread safety of TEDS' index buffer, but just increases the number of times it gets locked.

diff --git a/Engine/Plugins/Experimental/TypedElementsDataStorage/Source/TypedElementsDataStorage/Private/TypedElementDatabase.cpp b/Engine/Plugins/Experimental/TypedElementsDataStorage/Source/TypedElementsDataStorage/Private/TypedElementDatabase.cpp
--- a/Engine/Plugins/Experimental/TypedElementsDataStorage/Source/TypedElementsDataStorage/Private/TypedElementDatabase.cpp
+++ b/Engine/Plugins/Experimental/TypedElementsDataStorage/Source/TypedElementsDataStorage/Private/TypedElementDatabase.cpp
@@ -1040,7 +1040,8 @@ void UTypedElementDatabase::PreparePhase(EQueryTickPhase Phase, float DeltaTime)
 		}
 		// During the processing of queries no mutation can happen to the structure of the database, just fields being updated. As such
 		// it's safe to only take a shared lock
-		FGlobalLock::InternalSharedLock();
+		// TODO: This requires Mass to tell TEDS it's about to flush its deferred commands.
+		// FGlobalLock::InternalSharedLock();
 	}
 }
 
@@ -1052,7 +1053,9 @@ void UTypedElementDatabase::FinalizePhase(EQueryTickPhase Phase, float DeltaTime
 	{
 		// During the processing of queries no mutation can happen to the structure of the database, just fields being updated. As such
 		// it's safe to only take a shared lock
-		FGlobalLock::InternalSharedUnlock();
+		// TODO: This requires Mass to tell TEDS it's about to flush its deferred commands. Right now this gets called after the
+		// deferred commands are run, which require exclusive access.
+		//FGlobalLock::InternalSharedUnlock();
 		
 		// The preamble queries are all run on the game thread. While this is true it's safe to take a global write lock.
 		// If there's a performance loss because this lock is held too long, the work in RunPhasePostambleQueries can be split
