Introduced a new command queue in TEDS to support async TEDS Compat calls in the future.

diff --git a/Engine/Plugins/Experimental/TypedElementsDataStorage/Source/TypedElementsDataStorage/Private/Commands/EditorDataStorageCommandBuffer.h b/Engine/Plugins/Experimental/TypedElementsDataStorage/Source/TypedElementsDataStorage/Private/Commands/EditorDataStorageCommandBuffer.h
new file mode 100644
--- /dev/null
+++ b/Engine/Plugins/Experimental/TypedElementsDataStorage/Source/TypedElementsDataStorage/Private/Commands/EditorDataStorageCommandBuffer.h
@@ -0,0 +1,238 @@
+// Copyright Epic Games, Inc. All Rights Reserved.
+
+#pragma once
+
+#include <atomic>
+#include "Misc/TVariant.h"
+#include "Templates/Function.h"
+#include "TypedElementDatabaseScratchBuffer.h"
+
+namespace UE::EditorDataStorage
+{
+	/** Empty command that does nothing. Can be used as placeholder or to disable an existing command. */
+	struct FNopCommand final {};
+
+	/**
+	 * Class to store, optimize and process commands. Only the list of provided command objects can be used or TNopCommand which will is 
+	 * automatically added. Commands take the form of a struct with any number, including zero, member variables.
+	 * The typical use case of the command buffer is to repeatedly call `AddCommand`. This can be safely done from any thread. Periodically,
+	 * e.g. once per frame, `Collect` is called which collects all commands from all threads into a single local collection. This stores
+	 * references to the commands so doesn't move or copy the commands out of the scratch buffer where commands are created. Keep the
+	 * lifetime of objects in the scratch buffer in mind.
+	 * The returned collection can be further processed. This typically involves steps such as patching commands e.g. by resolving the
+	 * table index, sort commands in order, optimize commands to remove/replace/add commands as needed and finally the commands are
+	 * processed so they get executed. A utility class is also provided to help with optimizing the command buffer.
+	 */
+	template<typename... TCommand>
+	class FCommandBuffer final
+	{
+	public:
+		// Note that FNopCommand is always inserted at index 0. This means that any code using sizeof...(TCommand) needs to add 1 to the count.
+		using TCommandVariant = TVariant<FNopCommand, TCommand...>;
+
+		/** Storage for commands collected from all threads that haven't been executed yet. */
+		class FCollection final
+		{
+			friend class FCommandBuffer;
+			friend class FOptimizer;
+		public:
+			FCollection();
+
+			/** Adds a new command at the end of the command collection. */
+			template<typename T>
+			void AddCommand();
+			/** Adds a new command at the end of the command collection. */
+			template<typename T>
+			void AddCommand(T&& Command);
+
+			/**
+			 * Replaces the command at the given index. Prefer this over directly replacing the command variant to make sure tracking
+			 * in the buffer remains consistent.
+			 */
+			template<typename T>
+			void ReplaceCommand(int32 Index);
+			/**
+			 * Replaces the command at the given index. Prefer this over directly replacing the command variant to make sure tracking
+			 * in the buffer remains consistent.
+			 */
+			template<typename T>
+			void ReplaceCommand(int32 Index, T&& Command);
+
+			/** Inserts a new command before the provided index, moving the command at the given index one position up. */
+			template<typename T>
+			void InsertCommandBefore(int32 Index);
+
+			/** Inserts a new command before the provided index, moving the command at the given index one position up. */
+			template<typename T>
+			void InsertCommandBefore(int32 Index, T&& Command);
+
+			/**
+			 * Iterates over all collected commands, allowing commands to be updated.
+			 * Do not use this function to replace a command with new one as this will cause internal counters to mismatch.
+			 */
+			void ForEach(TFunctionRef<void(int32 Index, TCommandVariant& Command)> Iterator);
+
+			/** Sorts commands using the provided compare function. */
+			template<bool bStableSort = true>
+			void Sort(TFunctionRef<bool(const TCommandVariant& Lhs, const TCommandVariant& Rhs)> Comparer);
+
+			/**
+			 * Processes the local commands, that are retrieved after calling Collect or that still remain in the local queue. The
+			 * provided can be any object that has functions that can receive a comment. As an example
+			 *		struct FProcessor{ void operator()(const TNopCommand&) {} };
+			 */
+			template<typename TProcessor>
+			void Process(TProcessor&& Processor) const;
+
+			/** Resets the locally stored information, effectively clearing the command buffer until the next collect is called. */
+			void Reset();
+
+			/** Returns the number of instances of a command that are locally queued. */
+			template<typename Command>
+			uint32 GetCommandCount() const;
+
+			/** Returns the total number of commands, excluding nop operations. */
+			uint32 GetTotalCommandCount() const;
+		private:
+			TArray<TCommandVariant*> CommandReferences;
+			uint32 CommandCount[sizeof...(TCommand) + 1];
+			FTypedElementDatabaseScratchBuffer* ScratchBuffer = nullptr;
+		};
+
+		/** Utility class to help optimize a collection of commands. */
+		class FOptimizer
+		{
+		public:
+			using IsValidCallback = TFunction<bool(const TCommandVariant& Command)>;
+
+			explicit FOptimizer(FCommandBuffer::FCollection& InCommands);
+
+			TCommandVariant& GetLeft();
+			TCommandVariant& GetRight();
+
+			/**
+			 * Moves the left index one command to the right. If the right index is on the same command, right will also be moved one 
+			 * position to the right.
+			 * Before: 0 | L1 | R2 | 3  | 4
+			 * After:  0 | 1  | L2 | R3 | 4
+			 */
+			void MoveToNextLeft();
+			/**
+			 * Move the right index one command to the right
+			 * Before: 0 | L1 | R2 | 3  | 4
+			 * After:  0 | L1 | 2  | R3 | 4
+			 */
+			void MoveToNextRight();
+			/** 
+			 * Moves both the left and right index one command to the right
+			 * Before: 0 | L1 | 2  | R3 | 4
+			 * After:  0 | 1  | L2 | 3  | R4
+			 */
+			void MoveToNextLeftAndRight();
+			/**
+			 * Resets the right index to be one next to the left index.
+			 * Before: 0 | L1 | 2  | 3 | R4
+			 * After:  0 | L1 | R2 | 3 | 4
+			 */
+			void ResetRightNextToLeft();
+			/** 
+			 * Moves the left index to one less than the right index.
+			 * Before: L0 | 1 | 2  | R3 | 4
+			 * After:  0  | 1 | L2 | R3 | 4
+			 */
+			void MoveLeftBeforeRight();
+			/** 
+			 * Moves the left index to the right index and moves the right index one to the right.
+			 * Before: L0 | 1 | 2 | R3 | 4
+			 * After:  0  | 1 | 2 | L3 | R4
+			 */
+			void MoveLeftToRight();
+
+			/** Replaces the command at the left index with a new command. */
+			template<typename T>
+			void ReplaceLeft(T&& Command);
+			/** Replaces the command at the left index with a new command. */
+			template<typename T>
+			void ReplaceLeft();
+			/** Replaces the command at the right index with a new command. */
+			template<typename T>
+			void ReplaceRight(T&& Command);
+			/** Replaces the command at the right index with a new command. */
+			template<typename T>
+			void ReplaceRight();
+
+			/**
+			 * Inserts a new command before the command at the left index.
+			 * Before: 10 | L20 | 30  | R40 | 50
+			 * After:  10 | 15  | L20 | 30  | R40 | 50
+			 */
+			template<typename T>
+			void InsertBeforeLeft(T&& Command);
+			/**
+			 * Inserts a new command before the command at the left index.
+			 * Before: 10 | L20 | 30  | R40 | 50
+			 * After:  10 | 15  | L20 | 30  | R40 | 50
+			 */
+			template<typename T>
+			void InsertBeforeLeft();
+
+			/** Creates a new optimizer starting at the left of the current optimizer. */
+			FOptimizer BranchOnLeft();
+			/** Creates a new optimizer starting at the left of the current optimizer. The provides callback is used to check range. */
+			FOptimizer BranchOnLeft(IsValidCallback Callback);
+			/** Creates a new optimizer starting at the right of the current optimizer. */
+			FOptimizer BranchOnRight();
+			/** Creates a new optimizer starting at the right of the current optimizer. The provides callback is used to check range. */
+			FOptimizer BranchOnRight(IsValidCallback Callback);
+
+			bool IsValid() const;
+
+		private:
+			FOptimizer(FCommandBuffer::FCollection& InCommands, int32 InLeft);
+			FOptimizer(FCommandBuffer::FCollection& InCommands, int32 InLeft, IsValidCallback InCallback);
+
+			int32 MoveToNextNonNop(int32 Location);
+			int32 MoveToPreviousNonNop(int32 Location);
+
+			IsValidCallback IsValidCheck;
+			int32 Left;
+			int32 Right;
+			FCommandBuffer::FCollection& Commands;
+		};
+
+		void Initialize(FTypedElementDatabaseScratchBuffer& InScratchBuffer);
+
+		/** Adds a command the command buffer in a thread-safe manner. T has to be one of the provided commands or TNopCommand. */
+		template<typename T>
+		void AddCommand();
+
+		/** Adds a command the command buffer in a thread-safe manner. T has to be one of the provided commands or TNopCommand. */
+		template<typename T>
+		void AddCommand(T&& Command);
+
+		/**
+		 * Collects commands from all thread locally for further processing. This needs to be called every frame as commands are
+		 * stored in a temporary buffer that will get cleared after the frame.
+		 * Returns the number of collected commands.
+		 */
+		SIZE_T Collect(FCollection& Storage);
+
+	private:
+		// Wrapper object used to store a command in the scratch buffer.
+		struct FCommandInstance final
+		{
+			FCommandInstance() = default;
+			template <typename U, typename... TArgs>
+			explicit FCommandInstance(TInPlaceType<U>&&, TArgs&&... Args);
+
+			TCommandVariant Command;
+			std::atomic<FCommandInstance*> Next = nullptr;
+		};
+		FCommandInstance* CommandFront = nullptr;
+		std::atomic<FCommandInstance*> CommandBack = nullptr;
+
+		FTypedElementDatabaseScratchBuffer* ScratchBuffer = nullptr;
+	};
+} // namespace UE::EditorDataStorage
+
+#include "Commands/EditorDataStorageCommandBuffer.inl"
\ No newline at end of file
diff --git a/Engine/Plugins/Experimental/TypedElementsDataStorage/Source/TypedElementsDataStorage/Private/Commands/EditorDataStorageCommandBuffer.inl b/Engine/Plugins/Experimental/TypedElementsDataStorage/Source/TypedElementsDataStorage/Private/Commands/EditorDataStorageCommandBuffer.inl
new file mode 100644
--- /dev/null
+++ b/Engine/Plugins/Experimental/TypedElementsDataStorage/Source/TypedElementsDataStorage/Private/Commands/EditorDataStorageCommandBuffer.inl
@@ -0,0 +1,443 @@
+// Copyright Epic Games, Inc. All Rights Reserved.
+
+#include "Algo/BinarySearch.h"
+#include "HAL/UnrealMemory.h"
+
+namespace UE::EditorDataStorage
+{
+	//
+	// FCommandBuffer::FCollection
+	//
+
+	template<typename... TCommand>
+	FCommandBuffer<TCommand...>::FCollection::FCollection()
+	{
+		FMemory::Memset(CommandCount, 0);
+	}
+
+	template<typename... TCommand>
+	template<typename T>
+	void FCommandBuffer<TCommand...>::FCollection::AddCommand()
+	{
+		CommandReferences.Add(ScratchBuffer->Emplace<TCommandVariant>(TInPlaceType<T>()));
+		CommandCount[TCommandVariant::template IndexOfType<T>()]++;
+	}
+
+	template<typename... TCommand>
+	template<typename T>
+	void FCommandBuffer<TCommand...>::FCollection::AddCommand(T&& Command)
+	{
+		CommandReferences.Add(ScratchBuffer->Emplace<TCommandVariant>(TInPlaceType<T>(), Forward<T>(Command)));
+		CommandCount[TCommandVariant::template IndexOfType<T>()]++;
+	}
+
+	template<typename... TCommand>
+	template<typename T>
+	void FCommandBuffer<TCommand...>::FCollection::ReplaceCommand(int32 Index)
+	{
+		checkf(Index < CommandReferences.Num(), TEXT("Attempting to replace a command which is not in a valid position."));
+
+		TCommandVariant* OriginalCommand = CommandReferences[Index];
+		ensureMsgf(CommandCount[OriginalCommand->GetIndex()]-- > 0,
+			TEXT("A command is being replaced in TEDS command buffer isn't matching the expected count."));
+		CommandCount[TCommandVariant::template IndexOfType<T>()]++;
+		OriginalCommand->template Emplace<T>();
+	}
+
+	template<typename... TCommand>
+	template<typename T>
+	void FCommandBuffer<TCommand...>::FCollection::ReplaceCommand(int32 Index, T&& Command)
+	{
+		checkf(Index < CommandReferences.Num(), TEXT("Attempting to replace a command which is not in a valid position."));
+
+		TCommandVariant* OriginalCommand = CommandReferences[Index];
+		ensureMsgf(CommandCount[OriginalCommand->GetIndex()]-- > 0,
+			TEXT("A command is being replaced in TEDS command buffer isn't matching the expected count."));
+		CommandCount[TCommandVariant::template IndexOfType<T>()]++;
+		OriginalCommand->template Emplace<T>(Forward<T>(Command));
+	}
+
+	template<typename... TCommand>
+	template<typename T>
+	void FCommandBuffer<TCommand...>::FCollection::InsertCommandBefore(int32 Index)
+	{
+		TCommandVariant* CommandStorage = ScratchBuffer->Emplace<TCommandVariant>(TInPlaceType<T>());
+		CommandReferences.Insert(CommandStorage, Index);
+		CommandCount[TCommandVariant::template IndexOfType<T>()]++;
+	}
+
+	template<typename... TCommand>
+	template<typename T>
+	void FCommandBuffer<TCommand...>::FCollection::InsertCommandBefore(int32 Index, T&& Command)
+	{
+		TCommandVariant* CommandStorage = ScratchBuffer->Emplace<TCommandVariant>(TInPlaceType<T>(), Forward<T>(Command));
+		CommandReferences.Insert(CommandStorage, Index);
+		CommandCount[TCommandVariant::template IndexOfType<T>()]++;
+	}
+
+	template<typename... TCommand>
+	void FCommandBuffer<TCommand...>::FCollection::ForEach(TFunctionRef<void(int32 Index, TCommandVariant& Command)> Iterator)
+	{
+		// Don't use a local iterator or cache the size as the array can change during processing.
+		for (int32 Index = 0; Index < CommandReferences.Num(); ++Index)
+		{
+			Iterator(Index, *CommandReferences[Index]);
+		}
+	}
+
+	template<typename... TCommand>
+	template<bool bStableSort>
+	void FCommandBuffer<TCommand...>::FCollection::Sort(TFunctionRef<bool(const TCommandVariant&, const TCommandVariant&)> Comparer)
+	{
+		if constexpr (bStableSort)
+		{
+			CommandReferences.StableSort(Comparer);
+		}
+		else
+		{
+			CommandReferences.Sort(Comparer);
+		}
+	}
+
+	template<typename... TCommand>
+	template<typename TProcessor>
+	void FCommandBuffer<TCommand...>::FCollection::Process(TProcessor&& Processor) const
+	{
+		for (TCommandVariant* Command : CommandReferences)
+		{
+			Visit(Processor, *Command);
+		}
+	}
+
+	template<typename... TCommand>
+	void FCommandBuffer<TCommand...>::FCollection::Reset()
+	{
+		CommandReferences.Reset();
+		FMemory::Memset(CommandCount, 0);
+	}
+
+	template<typename... TCommand>
+	template<typename Command>
+	uint32 FCommandBuffer<TCommand...>::FCollection::GetCommandCount() const
+	{
+		return CommandCount[FCommandBuffer::TCommandVariant::template IndexOfType<Command>()];
+	}
+
+	template<typename... TCommand>
+	uint32 FCommandBuffer<TCommand...>::FCollection::GetTotalCommandCount() const
+	{
+		uint32 Result = 0;
+		for (uint32 Index = 1; Index < sizeof...(TCommand) + 1; ++Index)
+		{
+			Result += CommandCount[Index];
+		}
+		return Result;
+	}
+
+
+	//
+	// FCommandBuffer::FOptimizer
+	//
+
+	template<typename... TCommand>
+	FCommandBuffer<TCommand...>::FOptimizer::FOptimizer(FCommandBuffer::FCollection& InCommands)
+		: Left(0)
+		, Right(1)
+		, Commands(InCommands)
+	{
+		if (!Commands.CommandReferences.IsEmpty() && Commands.CommandReferences[0]->template IsType<FNopCommand>())
+		{
+			Left = MoveToNextNonNop(0);
+			Right = MoveToNextNonNop(Left);
+		}
+		else if (Commands.CommandReferences[1]->template IsType<FNopCommand>())
+		{
+			Right = MoveToNextNonNop(1);
+		}
+	}
+
+	template<typename... TCommand>
+	FCommandBuffer<TCommand...>::FOptimizer::FOptimizer(FCommandBuffer::FCollection& InCommands, int32 InLeft)
+		: Left(InLeft)
+		, Right(InLeft + 1)
+		, Commands(InCommands)
+	{
+		if (!Commands.CommandReferences.IsEmpty() && Commands.CommandReferences[Left]->template IsType<FNopCommand>())
+		{
+			Left = MoveToNextNonNop(Left);
+			Right = MoveToNextNonNop(Left);
+		}
+		else if (Commands.CommandReferences[Right]->template IsType<FNopCommand>())
+		{
+			Right = MoveToNextNonNop(Right);
+		}
+	}
+
+	template<typename... TCommand>
+	FCommandBuffer<TCommand...>::FOptimizer::FOptimizer(FCommandBuffer::FCollection& InCommands, int32 InLeft, IsValidCallback InCallback)
+		: IsValidCheck(MoveTemp(InCallback))
+		, Left(InLeft)
+		, Right(InLeft + 1)
+		, Commands(InCommands)
+	{
+		if (!Commands.CommandReferences.IsEmpty() && Commands.CommandReferences[Left]->template IsType<FNopCommand>())
+		{
+			Left = MoveToNextNonNop(Left);
+			Right = MoveToNextNonNop(Left);
+		}
+		else if (Commands.CommandReferences[Right]->template IsType<FNopCommand>())
+		{
+			Right = MoveToNextNonNop(Right);
+		}
+	}
+
+	template<typename... TCommand>
+	typename FCommandBuffer<TCommand...>::TCommandVariant& FCommandBuffer<TCommand...>::FOptimizer::GetLeft()
+	{
+		checkf(Left < Commands.CommandReferences.Num(), TEXT("Attempting to access an invalid left command from a command buffer."));
+		return *Commands.CommandReferences[Left];
+	}
+
+	template<typename... TCommand>
+	typename FCommandBuffer<TCommand...>::TCommandVariant& FCommandBuffer<TCommand...>::FOptimizer::GetRight()
+	{
+		checkf(Right < Commands.CommandReferences.Num(), TEXT("Attempting to access an invalid right command from a command buffer."));
+		return *Commands.CommandReferences[Right];
+	}
+
+	template<typename... TCommand>
+	int32 FCommandBuffer<TCommand...>::FOptimizer::MoveToNextNonNop(int32 Location)
+	{
+		int32 CommandCount = Commands.CommandReferences.Num();
+		do
+		{
+			++Location;
+		} while (Location < CommandCount && Commands.CommandReferences[Location]->template IsType<FNopCommand>());
+		return Location;
+	}
+
+	template<typename... TCommand>
+	int32 FCommandBuffer<TCommand...>::FOptimizer::MoveToPreviousNonNop(int32 Location)
+	{
+		while (Location > 0)
+		{
+			--Location;
+			if (!Commands.CommandReferences[Location]->template IsType<FNopCommand>())
+			{
+				break;
+			}
+		}
+		return Location;
+	}
+
+	template<typename... TCommand>
+	void FCommandBuffer<TCommand...>::FOptimizer::MoveToNextLeft()
+	{
+		Left = MoveToNextNonNop(Left);
+		if (Left == Right)
+		{
+			Right = MoveToNextNonNop(Right);
+		}
+	}
+
+	template<typename... TCommand>
+	void FCommandBuffer<TCommand...>::FOptimizer::MoveToNextRight()
+	{
+		Right = MoveToNextNonNop(Right);
+	}
+
+	template<typename... TCommand>
+	void FCommandBuffer<TCommand...>::FOptimizer::MoveToNextLeftAndRight()
+	{
+		Left = MoveToNextNonNop(Left);
+		Right = MoveToNextNonNop(Right);
+		if (Right == Left)
+		{
+			Right = MoveToNextNonNop(Right);
+		}
+	}
+
+	template<typename... TCommand>
+	void FCommandBuffer<TCommand...>::FOptimizer::ResetRightNextToLeft()
+	{
+		Right = Left;
+		Right = MoveToNextNonNop(Right);
+	}
+
+	template<typename... TCommand>
+	void FCommandBuffer<TCommand...>::FOptimizer::MoveLeftBeforeRight()
+	{
+		Left = MoveToPreviousNonNop(Right);
+	}
+
+	template<typename... TCommand>
+	void FCommandBuffer<TCommand...>::FOptimizer::MoveLeftToRight()
+	{
+		Left = Right;
+		Right = MoveToNextNonNop(Right);
+	}
+
+	template<typename... TCommand>
+	bool FCommandBuffer<TCommand...>::FOptimizer::IsValid() const
+	{
+		if (!IsValidCheck)
+		{
+			return Right < Commands.CommandReferences.Num();
+		}
+		else
+		{
+			return Right < Commands.CommandReferences.Num() ? IsValidCheck(*Commands.CommandReferences[Right]) : false;
+		}
+	}
+
+	template<typename... TCommand>
+	template<typename T>
+	void FCommandBuffer<TCommand...>::FOptimizer::ReplaceLeft(T&& Command)
+	{
+		Commands.ReplaceCommand(Left, Forward<T>(Command));
+	}
+
+	template<typename... TCommand>
+	template<typename T>
+	void FCommandBuffer<TCommand...>::FOptimizer::ReplaceLeft()
+	{
+		Commands.template ReplaceCommand<T>(Left);
+	}
+
+	template<typename... TCommand>
+	template<typename T>
+	void FCommandBuffer<TCommand...>::FOptimizer::ReplaceRight(T&& Command)
+	{
+		Commands.ReplaceCommand(Right, Forward<T>(Command));
+	}
+
+	template<typename... TCommand>
+	template<typename T>
+	void FCommandBuffer<TCommand...>::FOptimizer::ReplaceRight()
+	{
+		Commands.template ReplaceCommand<T>(Right);
+	}
+
+	template<typename... TCommand>
+	template<typename T>
+	void FCommandBuffer<TCommand...>::FOptimizer::InsertBeforeLeft(T&& Command)
+	{
+		Commands.InsertCommandBefore(Left, Forward<T>(Command));
+		++Left; // Move left to stay on the same object as before the insert.
+		++Right; // Move right one up as well so it stays on the same object.
+	}
+
+	template<typename... TCommand>
+	template<typename T>
+	void FCommandBuffer<TCommand...>::FOptimizer::InsertBeforeLeft()
+	{
+		Commands.template InsertCommandBefore<T>(Left);
+		++Left; // Move left to stay on the same object as before the insert.
+		++Right; // Move right one up as well so it stays on the same object.
+	}
+
+	template<typename... TCommand>
+	FCommandBuffer<TCommand...>::FOptimizer FCommandBuffer<TCommand...>::FOptimizer::BranchOnLeft()
+	{
+		return FOptimizer(Commands, Left);
+	}
+
+	template<typename... TCommand>
+	FCommandBuffer<TCommand...>::FOptimizer FCommandBuffer<TCommand...>::FOptimizer::BranchOnLeft(IsValidCallback Callback)
+	{
+		return FOptimizer(Commands, Left, MoveTemp(Callback));
+	}
+
+	template<typename... TCommand>
+	FCommandBuffer<TCommand...>::FOptimizer FCommandBuffer<TCommand...>::FOptimizer::BranchOnRight()
+	{
+		return FOptimizer(Commands, Right);
+	}
+
+	template<typename... TCommand>
+	FCommandBuffer<TCommand...>::FOptimizer FCommandBuffer<TCommand...>::FOptimizer::BranchOnRight(IsValidCallback Callback)
+	{
+		return FOptimizer(Commands, Right, MoveTemp(Callback));
+	}
+
+
+
+
+	//
+	// FCommandBuffer::FCommandInstance
+	//
+
+	template<typename... TCommand>
+	template <typename U, typename... TArgs>
+	FCommandBuffer<TCommand...>::FCommandInstance::FCommandInstance(TInPlaceType<U>&&, TArgs&&... Args)
+		: Command(TInPlaceType<U>(), Forward<TArgs>(Args)...)
+	{
+	}
+
+
+
+	//
+	// FCommandBuffer
+	//
+
+	template<typename... TCommand>
+	void FCommandBuffer<TCommand...>::Initialize(FTypedElementDatabaseScratchBuffer& InScratchBuffer)
+	{
+		ScratchBuffer = &InScratchBuffer;
+		CommandFront = ScratchBuffer->Emplace<FCommandInstance>();
+		CommandBack = CommandFront;
+	}
+
+	template<typename... TCommand>
+	template<typename T>
+	void FCommandBuffer<TCommand...>::AddCommand()
+	{
+		checkf(ScratchBuffer, TEXT("Attempting to add a command to the command buffer before it's been initialized."));
+
+		FCommandInstance* AllocatedCommand = ScratchBuffer->Emplace<FCommandInstance>(TInPlaceType<T>());
+		FCommandInstance* Previous = CommandBack.load();
+		while (!CommandBack.compare_exchange_weak(Previous, AllocatedCommand)) {}
+		Previous->Next = AllocatedCommand;
+	}
+
+	template<typename... TCommand>
+	template<typename T>
+	void FCommandBuffer<TCommand...>::AddCommand(T&& Command)
+	{
+		checkf(ScratchBuffer, TEXT("Attempting to add a command to the command buffer before it's been initialized."));
+
+		FCommandInstance* AllocatedCommand = ScratchBuffer->Emplace<FCommandInstance>(TInPlaceType<T>(), Forward<T>(Command));
+		FCommandInstance* Previous = CommandBack.load();
+		while (!CommandBack.compare_exchange_weak(Previous, AllocatedCommand)){}
+		Previous->Next = AllocatedCommand;
+	}
+
+	template<typename... TCommand>
+	SIZE_T FCommandBuffer<TCommand...>::Collect(FCollection& Storage)
+	{
+		checkf(ScratchBuffer, TEXT("Attempting to collect the command buffer before it's been initialized."));
+		Storage.ScratchBuffer = ScratchBuffer;
+
+		// Claim the command buffer by replacing it with a fresh one.
+		FCommandInstance* NewCommandFront = ScratchBuffer->Emplace<FCommandInstance>();
+		FCommandInstance* CurrentBack = CommandBack.load();
+		while (!CommandBack.compare_exchange_weak(CurrentBack, NewCommandFront)) {}
+
+		// Walk the list of commands and copy them locally. Skip the first entry as it will always be a no-op.
+		checkf(CommandFront->Command.template IsType<FNopCommand>(),
+			TEXT("The first operation in TEDS' command buffers should always be a no-op. ")
+			TEXT("As this is not the case, the buffer may have gotten corrupted."));
+		SIZE_T Result = 0;
+		FCommandInstance* Front = CommandFront->Next;
+		while (Front)
+		{
+			Storage.CommandReferences.Add(&Front->Command);
+			Storage.CommandCount[Front->Command.GetIndex()]++;
+			Front = Front->Next;
+			++Result;
+		}
+		CommandFront = NewCommandFront;
+
+		return Result;
+	}
+} // namespace UE::EditorDataStorage
\ No newline at end of file
diff --git a/Engine/Plugins/Experimental/TypedElementsDataStorage/Source/TypedElementsDataStorage/Tests/Commands/EditorDataStorageCommandBufferTests.cpp b/Engine/Plugins/Experimental/TypedElementsDataStorage/Source/TypedElementsDataStorage/Tests/Commands/EditorDataStorageCommandBufferTests.cpp
new file mode 100644
--- /dev/null
+++ b/Engine/Plugins/Experimental/TypedElementsDataStorage/Source/TypedElementsDataStorage/Tests/Commands/EditorDataStorageCommandBufferTests.cpp
@@ -0,0 +1,827 @@
+// Copyright Epic Games, Inc. All Rights Reserved.
+
+#if WITH_TESTS
+
+#include "TypedElementDatabaseScratchBuffer.h"
+#include "Commands/EditorDataStorageCommandBuffer.h"
+#include "Containers/ContainersFwd.h"
+#include "Tests/TestHarnessAdapter.h"
+
+struct FTestCommandA { int32 Value; };
+struct FTestCommandB { float Value; };
+struct FTestCommandC { bool Value; };
+enum class ETestCommandType
+{
+	Nop, A, B, C
+};
+
+using TestCommandBuffer = UE::EditorDataStorage::FCommandBuffer<FTestCommandA, FTestCommandB, FTestCommandC>;
+
+constexpr const TCHAR* TypeToString(ETestCommandType Type)
+{
+	switch (Type)
+	{
+	case ETestCommandType::Nop:
+		return TEXT("Nop");
+	case ETestCommandType::A:
+		return TEXT("A");
+	case ETestCommandType::B:
+		return TEXT("B");
+	case ETestCommandType::C:
+		return TEXT("C");
+	default:
+		return TEXT("<unknown>");
+	}
+}
+
+struct FSequenceTestingProcessor
+{
+	explicit FSequenceTestingProcessor(TArray<ETestCommandType> InCommandTypes)
+		: CommandTypes(MoveTemp(InCommandTypes))
+	{}
+
+	TArray<ETestCommandType> CommandTypes;
+	int Index = 0;
+
+	void operator()(const UE::EditorDataStorage::FNopCommand&)
+	{
+		REQUIRE_MESSAGE(TEXT("An additional NOP command was issued when no more commands were expected."), Index < CommandTypes.Num());
+		if (CommandTypes[Index] != ETestCommandType::Nop)
+		{
+			ADD_ERROR(FString::Printf(TEXT("Got nop command, but expected %s."), TypeToString(CommandTypes[Index])))
+		}
+		Index++;
+	}
+	void operator()(const FTestCommandA&)
+	{
+		REQUIRE_MESSAGE(TEXT("An additional A command was issued when no more commands were expected."), Index < CommandTypes.Num());
+		if (CommandTypes[Index] != ETestCommandType::A)
+		{
+			ADD_ERROR(FString::Printf(TEXT("Got A command, but expected %s."), TypeToString(CommandTypes[Index])))
+		}
+		Index++;
+	}
+	void operator()(const FTestCommandB&)
+	{
+		REQUIRE_MESSAGE(TEXT("An additional B command was issued when no more commands were expected."), Index < CommandTypes.Num());
+		if (CommandTypes[Index] != ETestCommandType::B)
+		{
+			ADD_ERROR(FString::Printf(TEXT("Got B command, but expected %s."), TypeToString(CommandTypes[Index])))
+		}
+		Index++;
+	}
+	void operator()(const FTestCommandC&)
+	{
+		REQUIRE_MESSAGE(TEXT("An additional C command was issued when no more commands were expected."), Index < CommandTypes.Num());
+		if (CommandTypes[Index] != ETestCommandType::C)
+		{
+			ADD_ERROR(FString::Printf(TEXT("Got C command, but expected %s."), TypeToString(CommandTypes[Index])))
+		}
+		Index++;
+	}
+};
+
+TEST_CASE_NAMED(EditorDataStorage_FCommandBuffer, "EditorDataStorage::Command Buffer", "[ApplicationContextMask][EngineFilter]")
+{
+	using namespace UE::EditorDataStorage;
+	
+	FTypedElementDatabaseScratchBuffer ScratchBuffer;
+
+	SECTION("Initialize")
+	{
+		TestCommandBuffer CommandBuffer;
+		CommandBuffer.Initialize(ScratchBuffer);
+	}
+
+	SECTION("AddCommand")
+	{
+		TestCommandBuffer CommandBuffer;
+		CommandBuffer.Initialize(ScratchBuffer);
+		CommandBuffer.AddCommand(FTestCommandA{ .Value = 42 });
+		ScratchBuffer.BatchDelete();
+	}
+
+	SECTION("Collect")
+	{
+		TestCommandBuffer CommandBuffer;
+		CommandBuffer.Initialize(ScratchBuffer);
+		CommandBuffer.AddCommand(FTestCommandA{ .Value = 42 });
+		TestCommandBuffer::FCollection PendingCommands;
+		SIZE_T CollectedCount = CommandBuffer.Collect(PendingCommands);
+		CHECK_EQUALS(TEXT("After adding one command only one should be collected."), CollectedCount, SIZE_T(1));
+		ScratchBuffer.BatchDelete();
+	}
+
+	SECTION("Collect with large number of commands.")
+	{
+		static constexpr int32 CommandCount = 10000;
+
+		TestCommandBuffer CommandBuffer;
+		CommandBuffer.Initialize(ScratchBuffer);
+		for (int32 Counter = 0; Counter < CommandCount; ++Counter)
+		{
+			CommandBuffer.AddCommand(FTestCommandA{ .Value = Counter });
+		}
+		TestCommandBuffer::FCollection PendingCommands;
+		SIZE_T CollectedCount = CommandBuffer.Collect(PendingCommands);
+		CHECK_EQUALS(TEXT("After adding one command only one should be collected."), CollectedCount, SIZE_T(CommandCount));
+		ScratchBuffer.BatchDelete();
+	}
+}
+
+TEST_CASE_NAMED(EditorDataStorage_FCommandBuffer_FCollection, "EditorDataStorage::Command Buffer::FCollection", "[ApplicationContextMask][EngineFilter]")
+{
+	using namespace UE::EditorDataStorage;
+
+	FTypedElementDatabaseScratchBuffer ScratchBuffer;
+	TestCommandBuffer CommandBuffer;
+	CommandBuffer.Initialize(ScratchBuffer);
+
+	SECTION("Add command")
+	{
+		TestCommandBuffer::FCollection Commands;
+		CommandBuffer.Collect(Commands);
+
+		CHECK_EQUALS(TEXT("An empty command collection should not have nops."), Commands.GetCommandCount<FNopCommand>(), 0);
+		Commands.AddCommand<FNopCommand>();
+		CHECK_EQUALS(TEXT("After adding a nop command there should be a naop."), Commands.GetCommandCount<FNopCommand>(), 1);
+	}
+
+	SECTION("Add command with argument")
+	{
+		TestCommandBuffer::FCollection Commands;
+		CommandBuffer.Collect(Commands);
+
+		CHECK_EQUALS(TEXT("An empty command collection should not have nops."), Commands.GetCommandCount<FTestCommandA>(), 0);
+		Commands.AddCommand(FTestCommandA{ .Value = 42 });
+		CHECK_EQUALS(TEXT("After adding a nop command there should be a naop."), Commands.GetCommandCount<FTestCommandA>(), 1);
+	}
+
+	SECTION("Replace command")
+	{
+		TestCommandBuffer::FCollection Commands;
+		CommandBuffer.Collect(Commands);
+
+		Commands.AddCommand(FTestCommandA{ .Value = 42 });
+		CHECK_EQUALS(TEXT("The test command wasn't in the collection."), Commands.GetCommandCount<FTestCommandA>(), 1);
+		CHECK_EQUALS(TEXT("An empty command collection should not have nops."), Commands.GetCommandCount<FNopCommand>(), 0);
+		
+		Commands.ReplaceCommand<FNopCommand>(0);
+		CHECK_EQUALS(TEXT("After replacing there shouldn't be a test command anymore."), Commands.GetCommandCount<FTestCommandA>(), 0);
+		CHECK_EQUALS(TEXT("After replacing with a nop command there should be a nop."), Commands.GetCommandCount<FNopCommand>(), 1);
+	}
+
+	SECTION("Replace command with argument")
+	{
+		TestCommandBuffer::FCollection Commands;
+		CommandBuffer.Collect(Commands);
+
+		Commands.AddCommand(FTestCommandA{ .Value = 42 });
+		CHECK_EQUALS(TEXT("Test command A wasn't in the collection."), Commands.GetCommandCount<FTestCommandA>(), 1);
+		CHECK_EQUALS(TEXT("An empty command collection should not have a test command B."), Commands.GetCommandCount<FTestCommandB>(), 0);
+		
+		Commands.ReplaceCommand(0, FTestCommandB{ .Value = 3.14f });
+		CHECK_EQUALS(TEXT("After replacing there shouldn't be a test command A anymore."), Commands.GetCommandCount<FTestCommandA>(), 0);
+		CHECK_EQUALS(TEXT("After replacing with there should be a test command B."), Commands.GetCommandCount<FTestCommandB>(), 1);
+	}
+
+	SECTION("InsertCommandBefore with argument")
+	{
+		TestCommandBuffer::FCollection Commands;
+		CommandBuffer.Collect(Commands);
+		Commands.AddCommand(FTestCommandA{ .Value = 10 });
+		Commands.AddCommand(FTestCommandA{ .Value = 20 });
+		Commands.AddCommand(FTestCommandA{ .Value = 30 });
+
+		Commands.InsertCommandBefore(1, FTestCommandB{ .Value = 42.0f });
+
+		FSequenceTestingProcessor Processor({ ETestCommandType::A, ETestCommandType::B, ETestCommandType::A, ETestCommandType::A });
+		Commands.Process(Processor);
+	}
+
+	SECTION("InsertCommandBefore with argument before first")
+	{
+		TestCommandBuffer::FCollection Commands;
+		CommandBuffer.Collect(Commands);
+		Commands.AddCommand(FTestCommandA{ .Value = 10 });
+		Commands.AddCommand(FTestCommandA{ .Value = 20 });
+		Commands.AddCommand(FTestCommandA{ .Value = 30 });
+
+		Commands.InsertCommandBefore(0, FTestCommandB{ .Value = 42.0f });
+
+		FSequenceTestingProcessor Processor({ ETestCommandType::B, ETestCommandType::A, ETestCommandType::A, ETestCommandType::A });
+		Commands.Process(Processor);
+	}
+
+	SECTION("Process")
+	{
+		TestCommandBuffer::FCollection Commands;
+		CommandBuffer.Collect(Commands);
+		Commands.AddCommand(FTestCommandA{ .Value = 42 });
+		Commands.AddCommand(FTestCommandB{ .Value = 42.0f });
+
+		struct FProcessor
+		{
+			int ACount = 0;
+			int BCount = 0;
+
+			void operator()(const FNopCommand&) { }
+			void operator()(const FTestCommandA&) { ++ACount; }
+			void operator()(const FTestCommandB&) { ++BCount; }
+			void operator()(const FTestCommandC&) { }
+		};
+		FProcessor Processor;
+		Commands.Process(Processor);
+
+		CHECK_EQUALS(TEXT("Expected test command A to be touched exactly once."), Processor.ACount, int(1));
+		CHECK_EQUALS(TEXT("Expected test command B to be touched exactly once."), Processor.BCount, int(1));
+	}
+
+	SECTION("InsertCommandBefore")
+	{
+		TestCommandBuffer::FCollection Commands;
+		CommandBuffer.Collect(Commands);
+		Commands.AddCommand(FTestCommandA{ .Value = 10 });
+		Commands.AddCommand(FTestCommandA{ .Value = 20 });
+		Commands.AddCommand(FTestCommandA{ .Value = 30 });
+
+		Commands.InsertCommandBefore<FNopCommand>(1);
+
+		FSequenceTestingProcessor Processor({ ETestCommandType::A, ETestCommandType::Nop, ETestCommandType::A, ETestCommandType::A });
+		Commands.Process(Processor);
+	}
+
+	SECTION("InsertCommandBefore before first")
+	{
+		TestCommandBuffer::FCollection Commands;
+		CommandBuffer.Collect(Commands);
+		Commands.AddCommand(FTestCommandA{ .Value = 10 });
+		Commands.AddCommand(FTestCommandA{ .Value = 20 });
+		Commands.AddCommand(FTestCommandA{ .Value = 30 });
+
+		Commands.InsertCommandBefore<FNopCommand>(0);
+
+		FSequenceTestingProcessor Processor({ ETestCommandType::Nop, ETestCommandType::A, ETestCommandType::A, ETestCommandType::A });
+		Commands.Process(Processor);
+	}
+
+	SECTION("ForEach")
+	{
+		TestCommandBuffer::FCollection Commands;
+		CommandBuffer.Collect(Commands);
+		Commands.AddCommand(FTestCommandA{ .Value = 42 });
+		Commands.AddCommand(FTestCommandB{ .Value = 42.0f });
+
+		struct FIterator
+		{
+			int ACount = 0;
+			int BCount = 0;
+
+			void operator()(const FNopCommand&) {}
+			void operator()(const FTestCommandA&) { ++ACount; }
+			void operator()(const FTestCommandB&) { ++BCount; }
+			void operator()(const FTestCommandC&) {}
+		};
+		FIterator Iterator;
+		Commands.ForEach(
+			[&Iterator](int32 Index, TestCommandBuffer::TCommandVariant& Command)
+			{
+				Visit(Iterator, Command);
+			});
+
+		CHECK_EQUALS(TEXT("Expected test command A to be touched exactly once."), Iterator.ACount, int(1));
+		CHECK_EQUALS(TEXT("Expected test command B to be touched exactly once."), Iterator.BCount, int(1));
+	}
+
+	SECTION("ForEach with culling")
+	{
+		TestCommandBuffer::FCollection Commands;
+		CommandBuffer.Collect(Commands);
+		Commands.AddCommand(FTestCommandA{ .Value = 42 });
+		Commands.AddCommand(FTestCommandB{ .Value = 42.0f });
+
+		Commands.ForEach(
+			[](int32 Index, TestCommandBuffer::TCommandVariant& Command)
+			{
+				struct FIterator
+				{
+					bool operator()(const FNopCommand&) { return true; }
+					bool operator()(const FTestCommandA& Command) { return false; }
+					bool operator()(const FTestCommandB&) { return true; }
+					bool operator()(const FTestCommandC&) { return true; }
+				};
+				if (!Visit(FIterator(), Command))
+				{
+					Command.Emplace<FNopCommand>();
+				}
+			});
+		struct FProcessor
+		{
+			int NopCount = 0;
+			int ACount = 0;
+			int BCount = 0;
+			int CCount = 0;
+
+			void operator()(const FNopCommand&) { ++NopCount; }
+			void operator()(const FTestCommandA&) { ++ACount; }
+			void operator()(const FTestCommandB&) { ++BCount; }
+			void operator()(const FTestCommandC&) { ++CCount; }
+		};
+		FProcessor Processor;
+		Commands.Process(Processor);
+
+		CHECK_EQUALS(TEXT("Expected test command Nop to be touched exactly once."), Processor.NopCount, int(1));
+		CHECK_EQUALS(TEXT("Expected test command A to not be touched."), Processor.ACount, int(0));
+		CHECK_EQUALS(TEXT("Expected test command B to be touched exactly once."), Processor.BCount, int(1));
+		CHECK_EQUALS(TEXT("Expected test command C to not be touched."), Processor.CCount, int(0));
+	}
+
+	SECTION("Sort")
+	{
+		TestCommandBuffer::FCollection Commands;
+		CommandBuffer.Collect(Commands);
+		Commands.AddCommand(FTestCommandB{ .Value = 1.0f });
+		Commands.AddCommand(FTestCommandA{ .Value = 1 });
+		Commands.AddCommand(FTestCommandB{ .Value = 2.0f });
+		Commands.AddCommand(FTestCommandA{ .Value = 2 });
+
+		Commands.Sort<true /* Use stable sort */>(
+			[](const TestCommandBuffer::TCommandVariant& Lhs, const TestCommandBuffer::TCommandVariant& Rhs)
+			{
+				return Lhs.GetIndex() < Rhs.GetIndex();
+			});
+
+		FSequenceTestingProcessor Processor({ ETestCommandType::A, ETestCommandType::A, ETestCommandType::B, ETestCommandType::B });
+		Commands.Process(Processor);
+	}
+
+	SECTION("Reset")
+	{
+		TestCommandBuffer::FCollection Commands;
+		CommandBuffer.Collect(Commands);
+		Commands.AddCommand(FTestCommandB{ .Value = 1.0f });
+		Commands.AddCommand(FTestCommandA{ .Value = 1 });
+		Commands.AddCommand(FTestCommandB{ .Value = 2.0f });
+		Commands.AddCommand(FTestCommandA{ .Value = 2 });
+
+		Commands.Reset();
+		uint32 Count = Commands.GetTotalCommandCount();
+		CHECK_EQUALS(TEXT("There are still commands in the command buffer after a reset."), Count, 0);
+
+		FSequenceTestingProcessor Processor({});
+		Commands.Process(Processor);
+	}
+
+	SECTION("GetCommandCount")
+	{
+		TestCommandBuffer::FCollection Commands;
+		CommandBuffer.Collect(Commands);
+		Commands.AddCommand(FTestCommandB{ .Value = 1.0f });
+		Commands.AddCommand(FTestCommandA{ .Value = 1 });
+		Commands.AddCommand(FTestCommandB{ .Value = 2.0f });
+		Commands.AddCommand(FTestCommandA{ .Value = 2 });
+
+		uint32 CountA = Commands.GetCommandCount<FTestCommandA>();
+		uint32 CountB = Commands.GetCommandCount<FTestCommandB>();
+		CHECK_EQUALS(TEXT("Total A count didn't match the number of added commands."), CountA, 2);
+		CHECK_EQUALS(TEXT("Total B count didn't match the number of added commands."), CountB, 2);
+	}
+
+	SECTION("GetTotalCommandCount")
+	{
+		TestCommandBuffer::FCollection Commands;
+		CommandBuffer.Collect(Commands);
+		Commands.AddCommand(FTestCommandB{ .Value = 1.0f });
+		Commands.AddCommand(FTestCommandA{ .Value = 1 });
+		Commands.AddCommand(FTestCommandB{ .Value = 2.0f });
+		Commands.AddCommand(FTestCommandA{ .Value = 2 });
+
+		uint32 Count = Commands.GetTotalCommandCount();
+		CHECK_EQUALS(TEXT("Total count didn't match the number of added commands."), Count, 4);
+	}
+}
+
+TEST_CASE_NAMED(EditorDataStorage_FCommandBuffer_FOptimizer, "EditorDataStorage::Command Buffer::FOptimizer", "[ApplicationContextMask][EngineFilter]")
+{
+	using namespace UE::EditorDataStorage;
+
+	FTypedElementDatabaseScratchBuffer ScratchBuffer;
+	TestCommandBuffer CommandBuffer;
+	CommandBuffer.Initialize(ScratchBuffer);
+
+	SECTION("Constructor corrects for nops before left")
+	{
+		TestCommandBuffer::FCollection Commands;
+		CommandBuffer.Collect(Commands);
+		Commands.AddCommand<FNopCommand>();
+		Commands.AddCommand(FTestCommandA{ .Value = 1 });
+		Commands.AddCommand(FTestCommandB{ .Value = 2.0f });
+
+		TestCommandBuffer::FOptimizer Optimizer(Commands);
+		bool LeftMatches = Optimizer.GetLeft().IsType<FTestCommandA>();
+		bool RightMatches = Optimizer.GetRight().IsType<FTestCommandB>();
+		CHECK_EQUALS(TEXT("Left was not correctly set."), LeftMatches, true);
+		CHECK_EQUALS(TEXT("Right was not correctly set."), RightMatches, true);
+	}
+
+	SECTION("Constructor corrects for nops before right")
+	{
+		TestCommandBuffer::FCollection Commands;
+		CommandBuffer.Collect(Commands);
+		Commands.AddCommand(FTestCommandA{ .Value = 1 });
+		Commands.AddCommand<FNopCommand>();
+		Commands.AddCommand(FTestCommandB{ .Value = 2.0f });
+
+		TestCommandBuffer::FOptimizer Optimizer(Commands);
+		bool LeftMatches = Optimizer.GetLeft().IsType<FTestCommandA>();
+		bool RightMatches = Optimizer.GetRight().IsType<FTestCommandB>();
+		CHECK_EQUALS(TEXT("Left was not correctly set."), LeftMatches, true);
+		CHECK_EQUALS(TEXT("Right was not correctly set."), RightMatches, true);
+	}
+
+	SECTION("GetLeft")
+	{
+		TestCommandBuffer::FCollection Commands;
+		CommandBuffer.Collect(Commands);
+		Commands.AddCommand(FTestCommandA{ .Value = 1 });
+		Commands.AddCommand(FTestCommandB{ .Value = 2.0f });
+
+		TestCommandBuffer::FOptimizer Optimizer(Commands);
+		bool LeftMatches = Optimizer.GetLeft().IsType<FTestCommandA>();
+		CHECK_EQUALS(TEXT("Left was not correctly retrieved."), LeftMatches, true);
+	}
+
+	SECTION("GetRight")
+	{
+		TestCommandBuffer::FCollection Commands;
+		CommandBuffer.Collect(Commands);
+		Commands.AddCommand(FTestCommandA{ .Value = 1 });
+		Commands.AddCommand(FTestCommandB{ .Value = 2.0f });
+
+		TestCommandBuffer::FOptimizer Optimizer(Commands);
+		bool RightMatches = Optimizer.GetRight().IsType<FTestCommandB>();
+		CHECK_EQUALS(TEXT("Right was not correctly retrieved."), RightMatches, true);
+	}
+
+	SECTION("MoveToNextLeft")
+	{
+		TestCommandBuffer::FCollection Commands;
+		CommandBuffer.Collect(Commands);
+		Commands.AddCommand(FTestCommandA{ .Value = 0 });
+		Commands.AddCommand(FTestCommandA{ .Value = 1 });
+		Commands.AddCommand(FTestCommandA{ .Value = 2 });
+
+		TestCommandBuffer::FOptimizer Optimizer(Commands);
+		Optimizer.MoveToNextLeft();
+		CHECK_EQUALS(TEXT("Left not moved correctly."), Optimizer.GetLeft().Get<FTestCommandA>().Value, 1);
+		CHECK_EQUALS(TEXT("Right not moved correctly."), Optimizer.GetRight().Get<FTestCommandA>().Value, 2);
+	}
+
+	SECTION("MoveToNextLeft but leave right as it's further out")
+	{
+		TestCommandBuffer::FCollection Commands;
+		CommandBuffer.Collect(Commands);
+		Commands.AddCommand(FTestCommandA{ .Value = 0 });
+		Commands.AddCommand(FTestCommandA{ .Value = 1 });
+		Commands.AddCommand(FTestCommandA{ .Value = 2 });
+		Commands.AddCommand(FTestCommandA{ .Value = 3 });
+
+		TestCommandBuffer::FOptimizer Optimizer(Commands);
+		Optimizer.MoveToNextRight(); // 1 -> 2
+		Optimizer.MoveToNextRight(); // 2 -> 3
+		Optimizer.MoveToNextLeft();
+		CHECK_EQUALS(TEXT("Left not moved correctly."), Optimizer.GetLeft().Get<FTestCommandA>().Value, 1);
+		CHECK_EQUALS(TEXT("Right was incorrectly moved."), Optimizer.GetRight().Get<FTestCommandA>().Value, 3);
+	}
+
+	SECTION("MoveToNextLeft and skip nops")
+	{
+		TestCommandBuffer::FCollection Commands;
+		CommandBuffer.Collect(Commands);
+		Commands.AddCommand(FTestCommandA{ .Value = 0 });
+		Commands.AddCommand<FNopCommand>();
+		Commands.AddCommand(FTestCommandA{ .Value = 2 });
+		Commands.AddCommand(FTestCommandA{ .Value = 3 });
+
+		TestCommandBuffer::FOptimizer Optimizer(Commands);
+		Optimizer.MoveToNextLeft();
+		CHECK_EQUALS(TEXT("Left not moved correctly."), Optimizer.GetLeft().Get<FTestCommandA>().Value, 2);
+		CHECK_EQUALS(TEXT("Right was incorrectly moved."), Optimizer.GetRight().Get<FTestCommandA>().Value, 3);
+	}
+
+	SECTION("MoveToNextRight")
+	{
+		TestCommandBuffer::FCollection Commands;
+		CommandBuffer.Collect(Commands);
+		Commands.AddCommand(FTestCommandA{ .Value = 0 });
+		Commands.AddCommand(FTestCommandA{ .Value = 1 });
+		Commands.AddCommand(FTestCommandA{ .Value = 2 });
+
+		TestCommandBuffer::FOptimizer Optimizer(Commands);
+		Optimizer.MoveToNextRight();
+		CHECK_EQUALS(TEXT("Left not moved correctly."), Optimizer.GetLeft().Get<FTestCommandA>().Value, 0);
+		CHECK_EQUALS(TEXT("Right not moved correctly."), Optimizer.GetRight().Get<FTestCommandA>().Value, 2);
+	}
+
+	SECTION("MoveToNextRight and skip nops")
+	{
+		TestCommandBuffer::FCollection Commands;
+		CommandBuffer.Collect(Commands);
+		Commands.AddCommand(FTestCommandA{ .Value = 0 });
+		Commands.AddCommand(FTestCommandA{ .Value = 1 });
+		Commands.AddCommand<FNopCommand>(); 
+		Commands.AddCommand(FTestCommandA{ .Value = 2 });
+
+		TestCommandBuffer::FOptimizer Optimizer(Commands);
+		Optimizer.MoveToNextRight();
+		CHECK_EQUALS(TEXT("Left not moved correctly."), Optimizer.GetLeft().Get<FTestCommandA>().Value, 0);
+		CHECK_EQUALS(TEXT("Right not moved correctly."), Optimizer.GetRight().Get<FTestCommandA>().Value, 2);
+	}
+
+	SECTION("ResetRightNextToLeft")
+	{
+		TestCommandBuffer::FCollection Commands;
+		CommandBuffer.Collect(Commands);
+		Commands.AddCommand(FTestCommandA{ .Value = 0 });
+		Commands.AddCommand(FTestCommandA{ .Value = 1 });
+		Commands.AddCommand(FTestCommandA{ .Value = 2 });
+
+		TestCommandBuffer::FOptimizer Optimizer(Commands);
+		Optimizer.MoveToNextRight(); // 1 -> 2
+		
+		Optimizer.ResetRightNextToLeft();
+		CHECK_EQUALS(TEXT("Left not moved correctly."), Optimizer.GetLeft().Get<FTestCommandA>().Value, 0);
+		CHECK_EQUALS(TEXT("Right not moved correctly."), Optimizer.GetRight().Get<FTestCommandA>().Value, 1);
+	}
+
+	SECTION("ResetRightNextToLeft and skip nops")
+	{
+		TestCommandBuffer::FCollection Commands;
+		CommandBuffer.Collect(Commands);
+		Commands.AddCommand(FTestCommandA{ .Value = 0 });
+		Commands.AddCommand<FNopCommand>();
+		Commands.AddCommand(FTestCommandA{ .Value = 1 });
+		Commands.AddCommand(FTestCommandA{ .Value = 2 });
+
+		TestCommandBuffer::FOptimizer Optimizer(Commands);
+		Optimizer.MoveToNextRight(); // 1 -> 2
+
+		Optimizer.ResetRightNextToLeft();
+		CHECK_EQUALS(TEXT("Left not moved correctly."), Optimizer.GetLeft().Get<FTestCommandA>().Value, 0);
+		CHECK_EQUALS(TEXT("Right not moved correctly."), Optimizer.GetRight().Get<FTestCommandA>().Value, 1);
+	}
+
+	SECTION("MoveToNextLeftAndRight")
+	{
+		TestCommandBuffer::FCollection Commands;
+		CommandBuffer.Collect(Commands);
+		Commands.AddCommand(FTestCommandA{ .Value = 0 });
+		Commands.AddCommand(FTestCommandA{ .Value = 1 });
+		Commands.AddCommand(FTestCommandA{ .Value = 2 });
+		Commands.AddCommand(FTestCommandA{ .Value = 3 });
+
+		TestCommandBuffer::FOptimizer Optimizer(Commands);
+		Optimizer.MoveToNextRight(); // 1 -> 2
+
+		Optimizer.MoveToNextLeftAndRight();
+		CHECK_EQUALS(TEXT("Left not moved correctly."), Optimizer.GetLeft().Get<FTestCommandA>().Value, 1);
+		CHECK_EQUALS(TEXT("Right not moved correctly."), Optimizer.GetRight().Get<FTestCommandA>().Value, 3);
+	}
+
+	SECTION("MoveToNextLeftAndRight and skip nops")
+	{
+		TestCommandBuffer::FCollection Commands;
+		CommandBuffer.Collect(Commands);
+		Commands.AddCommand(FTestCommandA{ .Value = 0 });
+		Commands.AddCommand<FNopCommand>();
+		Commands.AddCommand(FTestCommandA{ .Value = 1 });
+		Commands.AddCommand(FTestCommandA{ .Value = 2 });
+		Commands.AddCommand<FNopCommand>();
+		Commands.AddCommand(FTestCommandA{ .Value = 3 });
+
+		TestCommandBuffer::FOptimizer Optimizer(Commands);
+		Optimizer.MoveToNextRight(); // 1 -> 2
+
+		Optimizer.MoveToNextLeftAndRight();
+		CHECK_EQUALS(TEXT("Left not moved correctly."), Optimizer.GetLeft().Get<FTestCommandA>().Value, 1);
+		CHECK_EQUALS(TEXT("Right not moved correctly."), Optimizer.GetRight().Get<FTestCommandA>().Value, 3);
+	}
+
+	SECTION("MoveLeftBeforeRight")
+	{
+		TestCommandBuffer::FCollection Commands;
+		CommandBuffer.Collect(Commands);
+		Commands.AddCommand(FTestCommandA{ .Value = 0 });
+		Commands.AddCommand(FTestCommandA{ .Value = 1 });
+		Commands.AddCommand(FTestCommandA{ .Value = 2 });
+		
+		TestCommandBuffer::FOptimizer Optimizer(Commands);
+		Optimizer.MoveToNextRight(); // 1 -> 2
+
+		Optimizer.MoveLeftBeforeRight();
+		CHECK_EQUALS(TEXT("Left not moved correctly."), Optimizer.GetLeft().Get<FTestCommandA>().Value, 1);
+		CHECK_EQUALS(TEXT("Right not moved correctly."), Optimizer.GetRight().Get<FTestCommandA>().Value, 2);
+	}
+
+	SECTION("MoveLeftBeforeRight skip nops")
+	{
+		TestCommandBuffer::FCollection Commands;
+		CommandBuffer.Collect(Commands);
+		Commands.AddCommand(FTestCommandA{ .Value = 0 });
+		Commands.AddCommand(FTestCommandA{ .Value = 1 });
+		Commands.AddCommand(FTestCommandA{ .Value = 2 });
+		Commands.AddCommand<FNopCommand>();
+		Commands.AddCommand(FTestCommandA{ .Value = 3 });
+
+		TestCommandBuffer::FOptimizer Optimizer(Commands);
+		Optimizer.MoveToNextRight(); // 1 -> 2
+		Optimizer.MoveToNextRight(); // 2 -> 3
+
+		Optimizer.MoveLeftBeforeRight();
+		CHECK_EQUALS(TEXT("Left not moved correctly."), Optimizer.GetLeft().Get<FTestCommandA>().Value, 2);
+		CHECK_EQUALS(TEXT("Right not moved correctly."), Optimizer.GetRight().Get<FTestCommandA>().Value, 3);
+	}
+
+	SECTION("MoveLeftToRight")
+	{
+		TestCommandBuffer::FCollection Commands;
+		CommandBuffer.Collect(Commands);
+		Commands.AddCommand(FTestCommandA{ .Value = 0 });
+		Commands.AddCommand(FTestCommandA{ .Value = 1 });
+		Commands.AddCommand(FTestCommandA{ .Value = 2 });
+		Commands.AddCommand(FTestCommandA{ .Value = 3 });
+
+		TestCommandBuffer::FOptimizer Optimizer(Commands);
+		Optimizer.MoveToNextRight(); // 1 -> 2
+
+		Optimizer.MoveLeftToRight();
+		CHECK_EQUALS(TEXT("Left not moved correctly."), Optimizer.GetLeft().Get<FTestCommandA>().Value, 2);
+		CHECK_EQUALS(TEXT("Right not moved correctly."), Optimizer.GetRight().Get<FTestCommandA>().Value, 3);
+	}
+
+	SECTION("MoveLeftToRight skip nops")
+	{
+		TestCommandBuffer::FCollection Commands;
+		CommandBuffer.Collect(Commands);
+		Commands.AddCommand(FTestCommandA{ .Value = 0 });
+		Commands.AddCommand(FTestCommandA{ .Value = 1 });
+		Commands.AddCommand(FTestCommandA{ .Value = 2 });
+		Commands.AddCommand<FNopCommand>();
+		Commands.AddCommand(FTestCommandA{ .Value = 3 });
+
+		TestCommandBuffer::FOptimizer Optimizer(Commands);
+		Optimizer.MoveToNextRight(); // 1 -> 2
+
+		Optimizer.MoveLeftToRight();
+		CHECK_EQUALS(TEXT("Left not moved correctly."), Optimizer.GetLeft().Get<FTestCommandA>().Value, 2);
+		CHECK_EQUALS(TEXT("Right not moved correctly."), Optimizer.GetRight().Get<FTestCommandA>().Value, 3);
+	}
+
+	SECTION("ReplaceLeft")
+	{
+		TestCommandBuffer::FCollection Commands;
+		CommandBuffer.Collect(Commands);
+		Commands.AddCommand(FTestCommandA{ .Value = 1 });
+		Commands.AddCommand(FTestCommandA{ .Value = 2 });
+		
+		TestCommandBuffer::FOptimizer Optimizer(Commands);
+		Optimizer.ReplaceLeft<FNopCommand>();
+
+		FSequenceTestingProcessor Processor({ ETestCommandType::Nop, ETestCommandType::A });
+		Commands.Process(Processor);
+	}
+
+	SECTION("ReplaceLeft with argument")
+	{
+		TestCommandBuffer::FCollection Commands;
+		CommandBuffer.Collect(Commands);
+		Commands.AddCommand(FTestCommandA{ .Value = 1 });
+		Commands.AddCommand(FTestCommandA{ .Value = 2 });
+
+		TestCommandBuffer::FOptimizer Optimizer(Commands);
+		Optimizer.ReplaceLeft(FTestCommandB{ .Value = 3.14f });
+
+		FSequenceTestingProcessor Processor({ ETestCommandType::B, ETestCommandType::A });
+		Commands.Process(Processor);
+	}
+
+	SECTION("ReplaceRight")
+	{
+		TestCommandBuffer::FCollection Commands;
+		CommandBuffer.Collect(Commands);
+		Commands.AddCommand(FTestCommandA{ .Value = 1 });
+		Commands.AddCommand(FTestCommandA{ .Value = 2 });
+
+		TestCommandBuffer::FOptimizer Optimizer(Commands);
+		Optimizer.ReplaceRight<FNopCommand>();
+
+		FSequenceTestingProcessor Processor({ ETestCommandType::A, ETestCommandType::Nop });
+		Commands.Process(Processor);
+	}
+
+	SECTION("ReplaceRight with argument")
+	{
+		TestCommandBuffer::FCollection Commands;
+		CommandBuffer.Collect(Commands);
+		Commands.AddCommand(FTestCommandA{ .Value = 1 });
+		Commands.AddCommand(FTestCommandA{ .Value = 2 });
+
+		TestCommandBuffer::FOptimizer Optimizer(Commands);
+		Optimizer.ReplaceRight(FTestCommandB{ .Value = 3.14f });
+
+		FSequenceTestingProcessor Processor({ ETestCommandType::A, ETestCommandType::B });
+		Commands.Process(Processor);
+	}
+
+	SECTION("InsertBeforeLeft")
+	{
+		TestCommandBuffer::FCollection Commands;
+		CommandBuffer.Collect(Commands);
+		Commands.AddCommand(FTestCommandA{ .Value = 1 });
+		Commands.AddCommand(FTestCommandA{ .Value = 2 });
+		Commands.AddCommand(FTestCommandA{ .Value = 3 });
+
+		TestCommandBuffer::FOptimizer Optimizer(Commands);
+		Optimizer.MoveToNextLeft(); // 1 -> 2
+		Optimizer.InsertBeforeLeft<FNopCommand>();
+
+		FSequenceTestingProcessor Processor({ ETestCommandType::A, ETestCommandType::Nop, ETestCommandType::A, ETestCommandType::A });
+		Commands.Process(Processor);
+	}
+
+	SECTION("InsertBeforeLeft with argument")
+	{
+		TestCommandBuffer::FCollection Commands;
+		CommandBuffer.Collect(Commands);
+		Commands.AddCommand(FTestCommandA{ .Value = 1 });
+		Commands.AddCommand(FTestCommandA{ .Value = 2 });
+		Commands.AddCommand(FTestCommandA{ .Value = 3 });
+
+		TestCommandBuffer::FOptimizer Optimizer(Commands);
+		Optimizer.MoveToNextLeft(); // 1 -> 2
+		Optimizer.InsertBeforeLeft(FTestCommandB{ .Value = 3.14f });
+
+		FSequenceTestingProcessor Processor({ ETestCommandType::A, ETestCommandType::B, ETestCommandType::A, ETestCommandType::A });
+		Commands.Process(Processor);
+	}
+
+	SECTION("InsertBeforeLeft before first")
+	{
+		TestCommandBuffer::FCollection Commands;
+		CommandBuffer.Collect(Commands);
+		Commands.AddCommand(FTestCommandA{ .Value = 1 });
+		Commands.AddCommand(FTestCommandA{ .Value = 2 });
+		
+		TestCommandBuffer::FOptimizer Optimizer(Commands);
+		Optimizer.InsertBeforeLeft<FNopCommand>();
+
+		FSequenceTestingProcessor Processor({ ETestCommandType::Nop, ETestCommandType::A, ETestCommandType::A });
+		Commands.Process(Processor);
+	}
+
+	SECTION("InsertBeforeLeft before first with argument")
+	{
+		TestCommandBuffer::FCollection Commands;
+		CommandBuffer.Collect(Commands);
+		Commands.AddCommand(FTestCommandA{ .Value = 1 });
+		Commands.AddCommand(FTestCommandA{ .Value = 2 });
+
+		TestCommandBuffer::FOptimizer Optimizer(Commands);
+		Optimizer.InsertBeforeLeft(FTestCommandB{ .Value = 3.14f });
+
+		FSequenceTestingProcessor Processor({ ETestCommandType::B, ETestCommandType::A, ETestCommandType::A });
+		Commands.Process(Processor);
+	}
+
+	SECTION("BranchOnLeft")
+	{
+		TestCommandBuffer::FCollection Commands;
+		CommandBuffer.Collect(Commands);
+		Commands.AddCommand(FTestCommandA{ .Value = 1 });
+		Commands.AddCommand(FTestCommandA{ .Value = 2 });
+		Commands.AddCommand(FTestCommandA{ .Value = 3 });
+
+		TestCommandBuffer::FOptimizer Optimizer(Commands);
+		Optimizer.MoveToNextLeftAndRight(); // 1->2, 2->3
+
+		TestCommandBuffer::FOptimizer BranchedOptimizer = Optimizer.BranchOnLeft();
+		CHECK_EQUALS(TEXT("Left not moved correctly."), BranchedOptimizer.GetLeft().Get<FTestCommandA>().Value, 2);
+		CHECK_EQUALS(TEXT("Right not moved correctly."), BranchedOptimizer.GetRight().Get<FTestCommandA>().Value, 3);
+	}
+
+	SECTION("BranchOnRight")
+	{
+		TestCommandBuffer::FCollection Commands;
+		CommandBuffer.Collect(Commands);
+		Commands.AddCommand(FTestCommandA{ .Value = 1 });
+		Commands.AddCommand(FTestCommandA{ .Value = 2 });
+		Commands.AddCommand(FTestCommandA{ .Value = 3 });
+
+		TestCommandBuffer::FOptimizer Optimizer(Commands);
+		
+		TestCommandBuffer::FOptimizer BranchedOptimizer = Optimizer.BranchOnRight();
+		CHECK_EQUALS(TEXT("Left not moved correctly."), BranchedOptimizer.GetLeft().Get<FTestCommandA>().Value, 2);
+		CHECK_EQUALS(TEXT("Right not moved correctly."), BranchedOptimizer.GetRight().Get<FTestCommandA>().Value, 3);
+	}
+}
+#endif // WITH_TESTS
