Fixed Mac compile error.

diff --git a/Engine/Plugins/Experimental/TypedElementsDataStorage/Source/TypedElementsDataStorage/Private/Commands/EditorDataStorageCommandBuffer.inl b/Engine/Plugins/Experimental/TypedElementsDataStorage/Source/TypedElementsDataStorage/Private/Commands/EditorDataStorageCommandBuffer.inl
--- a/Engine/Plugins/Experimental/TypedElementsDataStorage/Source/TypedElementsDataStorage/Private/Commands/EditorDataStorageCommandBuffer.inl
+++ b/Engine/Plugins/Experimental/TypedElementsDataStorage/Source/TypedElementsDataStorage/Private/Commands/EditorDataStorageCommandBuffer.inl
@@ -337,25 +337,25 @@ namespace UE::EditorDataStorage
 	}
 
 	template<typename... TCommand>
-	FCommandBuffer<TCommand...>::FOptimizer FCommandBuffer<TCommand...>::FOptimizer::BranchOnLeft()
+	typename FCommandBuffer<TCommand...>::FOptimizer FCommandBuffer<TCommand...>::FOptimizer::BranchOnLeft()
 	{
 		return FOptimizer(Commands, Left);
 	}
 
 	template<typename... TCommand>
-	FCommandBuffer<TCommand...>::FOptimizer FCommandBuffer<TCommand...>::FOptimizer::BranchOnLeft(IsValidCallback Callback)
+	typename FCommandBuffer<TCommand...>::FOptimizer FCommandBuffer<TCommand...>::FOptimizer::BranchOnLeft(IsValidCallback Callback)
 	{
 		return FOptimizer(Commands, Left, MoveTemp(Callback));
 	}
 
 	template<typename... TCommand>
-	FCommandBuffer<TCommand...>::FOptimizer FCommandBuffer<TCommand...>::FOptimizer::BranchOnRight()
+	typename FCommandBuffer<TCommand...>::FOptimizer FCommandBuffer<TCommand...>::FOptimizer::BranchOnRight()
 	{
 		return FOptimizer(Commands, Right);
 	}
 
 	template<typename... TCommand>
-	FCommandBuffer<TCommand...>::FOptimizer FCommandBuffer<TCommand...>::FOptimizer::BranchOnRight(IsValidCallback Callback)
+	typename FCommandBuffer<TCommand...>::FOptimizer FCommandBuffer<TCommand...>::FOptimizer::BranchOnRight(IsValidCallback Callback)
 	{
 		return FOptimizer(Commands, Right, MoveTemp(Callback));
 	}
