Turn on GC integration with TEDS Compat permanently.

This change enables TEDS Compat monitoring the Garbage Collection (GC) for updates so it has an last-minute opportunity to remove objects from its internal state before they get deleted. This can be disabled using `TEDS.Feature.IntegrateWithGC`.  For testing purposes there were two calls added for deleting actors and entities exclusively from the the GC. While this was useful for testing it doesn't work for objects that were deleted without forcing a GC collection pass, so these were removed as testing has completed.

diff --git a/Engine/Plugins/Experimental/TypedElementsDataStorage/Source/TypedElementsDataStorage/Private/TypedElementDatabaseCompatibility.cpp b/Engine/Plugins/Experimental/TypedElementsDataStorage/Source/TypedElementsDataStorage/Private/TypedElementDatabaseCompatibility.cpp
--- a/Engine/Plugins/Experimental/TypedElementsDataStorage/Source/TypedElementsDataStorage/Private/TypedElementDatabaseCompatibility.cpp
+++ b/Engine/Plugins/Experimental/TypedElementsDataStorage/Source/TypedElementsDataStorage/Private/TypedElementDatabaseCompatibility.cpp
@@ -23,21 +23,14 @@
 
 namespace TypedElementDataStorage
 {
-	bool bActorsClearedByGC = false;
-	FAutoConsoleVariableRef CVarActorsClearedByGC(
-		TEXT("TEDS.Feature.ActorsClearedByGC"),
-		bActorsClearedByGC,
-		TEXT("Enables actors being removed through the garbage collection instead of requiring explicit removal. Enables `TEDS.Feature.IntegrateWithGC`."));
-
-	bool bEntitiesClearedByGC = false;
-	FAutoConsoleVariableRef CVarEntitiesClearedByGC(
-		TEXT("TEDS.Feature.EntitiesClearedByGC"),
-		bEntitiesClearedByGC,
-		TEXT("Enables entities being removed through the garbage collection instead of requiring explicit removal. Enables `TEDS.Feature.IntegrateWithGC`."));
+	bool bIntegrateWithGC = true;
+	FAutoConsoleVariableRef CVarIntegrateWithGC(
+		TEXT("TEDS.Feature.IntegrateWithGC"),
+		bIntegrateWithGC,
+		TEXT("Enabled monitoring GC events from TEDS Compat so objects are cleared up at the last minute, before they're destroyed."));
 }
 
-static const FName ActorsClearedByGCExtensionName(TEXT("ActorsClearedByGCExtension"));
-static const FName EntitiesClearedByGCExtensionName(TEXT("EntitiesClearedByGCExtension"));
+static const FName IntegrateWithGCName(TEXT("IntegrateWithGC"));
 
 void UTypedElementDatabaseCompatibility::Initialize(UTypedElementDatabase* InStorage)
 {
@@ -204,13 +197,9 @@ bool UTypedElementDatabaseCompatibility::SupportsExtension(FName Extension) cons
 {
 	using namespace TypedElementDataStorage;
 
-	if (Extension == ActorsClearedByGCExtensionName)
-	{
-		return bActorsClearedByGC;
-	}
-	else if (Extension == EntitiesClearedByGCExtensionName)
+	if (Extension == IntegrateWithGCName)
 	{
-		return bEntitiesClearedByGC;
+		return bIntegrateWithGC;
 	}
 	else
 	{
@@ -222,13 +211,9 @@ void UTypedElementDatabaseCompatibility::ListExtensions(TFunctionRef<void(FName)
 {
 	using namespace TypedElementDataStorage;
 
-	if (bActorsClearedByGC)
-	{
-		Callback(ActorsClearedByGCExtensionName);
-	}
-	if (bEntitiesClearedByGC)
+	if (bIntegrateWithGC)
 	{
-		Callback(EntitiesClearedByGCExtensionName);
+		Callback(IntegrateWithGCName);
 	}
 }
 
@@ -817,7 +802,7 @@ void UTypedElementDatabaseCompatibility::OnPostGcUnreachableAnalysis()
 	using namespace TypedElementDataStorage;
 	using namespace TypedElementQueryBuilder;
 
-	if (bActorsClearedByGC || bEntitiesClearedByGC)
+	if (bIntegrateWithGC)
 	{
 		TArray<TPair<FUObjectItem*, RowHandle>> DeletedObjects;
 
@@ -847,30 +832,22 @@ void UTypedElementDatabaseCompatibility::OnPostGcUnreachableAnalysis()
 
 void UTypedElementDatabaseCompatibility::OnPostWorldInitialization(UWorld* World, const UWorld::InitializationValues InitializationValues)
 {
-	if (!TypedElementDataStorage::bActorsClearedByGC)
-	{
-		FDelegateHandle Handle = World->AddOnActorDestroyedHandler(
-			FOnActorDestroyed::FDelegate::CreateUObject(this, &UTypedElementDatabaseCompatibility::OnActorDestroyed));
-		ActorDestroyedDelegateHandles.Add(World, Handle);
-	}
+	FDelegateHandle Handle = World->AddOnActorDestroyedHandler(
+		FOnActorDestroyed::FDelegate::CreateUObject(this, &UTypedElementDatabaseCompatibility::OnActorDestroyed));
+	ActorDestroyedDelegateHandles.Add(World, Handle);
 }
 
 void UTypedElementDatabaseCompatibility::OnPreWorldFinishDestroy(UWorld* World)
 {
-	if (!TypedElementDataStorage::bActorsClearedByGC)
+	FDelegateHandle Handle;
+	if (ActorDestroyedDelegateHandles.RemoveAndCopyValue(World, Handle))
 	{
-		FDelegateHandle Handle;
-		if (ActorDestroyedDelegateHandles.RemoveAndCopyValue(World, Handle))
-		{
-			World->RemoveOnActorDestroyededHandler(Handle);
-		}
+		World->RemoveOnActorDestroyededHandler(Handle);
 	}
 }
 
 void UTypedElementDatabaseCompatibility::OnActorDestroyed(AActor* Actor)
 {
-	checkf(!TypedElementDataStorage::bActorsClearedByGC, 
-		TEXT("If actors are cleared from TEDS Compatibility through garbage collection, there shouldn't be any explicitly actor destruction handlers."));
 	RemoveCompatibleObjectExplicit(Actor);
 }
 
diff --git a/Engine/Source/Editor/UnrealEd/Private/LevelEditorViewport.cpp b/Engine/Source/Editor/UnrealEd/Private/LevelEditorViewport.cpp
--- a/Engine/Source/Editor/UnrealEd/Private/LevelEditorViewport.cpp
+++ b/Engine/Source/Editor/UnrealEd/Private/LevelEditorViewport.cpp
@@ -1670,10 +1670,8 @@ void FLevelEditorViewportClient::DestroyDropPreviewElements()
 	// deleting preview actors in UUnrealEdEngine::DeleteActors skips explicit handle deregistration,
 	// and although it does still happen in the immediately triggered garbage cleanup, that feels
 	// potentially brittle.
-	static FName ActorsClearedByGCExtensionName(TEXT("ActorsClearedByGCExtension"));
 	UTypedElementRegistry* TypedElementRegistry = UTypedElementRegistry::GetInstance();
-	if (ITypedElementDataStorageCompatibilityInterface* TedsCompat = TypedElementRegistry->GetMutableDataStorageCompatibility();
-		TedsCompat && !TedsCompat->SupportsExtension(ActorsClearedByGCExtensionName))
+	if (ITypedElementDataStorageCompatibilityInterface* TedsCompat = TypedElementRegistry->GetMutableDataStorageCompatibility())
 	{
 		DropPreviewElements->ForEachElement<ITypedElementObjectInterface>(
 			[this, TedsCompat](const TTypedElement<ITypedElementObjectInterface>& InElement)
diff --git a/Engine/Source/Runtime/Engine/Private/Elements/Framework/EngineElementsLibrary.cpp b/Engine/Source/Runtime/Engine/Private/Elements/Framework/EngineElementsLibrary.cpp
--- a/Engine/Source/Runtime/Engine/Private/Elements/Framework/EngineElementsLibrary.cpp
+++ b/Engine/Source/Runtime/Engine/Private/Elements/Framework/EngineElementsLibrary.cpp
@@ -389,11 +389,8 @@ void UEngineElementsLibrary::RegisterActorElement(const AActor* InActor)
 
 void UEngineElementsLibrary::UnregisterActorElement(const AActor* InActor)
 {
-	static FName ActorsClearedByGCExtensionName = TEXT("ActorsClearedByGCExtension");
-
 	UTypedElementRegistry* Registry = UTypedElementRegistry::GetInstance();
-	ITypedElementDataStorageCompatibilityInterface* Storage = Registry->GetMutableDataStorageCompatibility();
-	if (Storage && !Storage->SupportsExtension(ActorsClearedByGCExtensionName))
+	if (ITypedElementDataStorageCompatibilityInterface* Storage = Registry->GetMutableDataStorageCompatibility())
 	{
 		Storage->RemoveCompatibleObject(const_cast<AActor*>(InActor));
 	}
