[TEDS] Removing the hotfixable cvar for disabling TEDS. This has been on for a few releases now and new features built on top of the core will have their own cvar moving forward

#rnx

diff --git a/Engine/Plugins/Experimental/TypedElementsDataStorage/Source/TypedElementsDataStorage/Private/TypedElementsDataStorage.cpp b/Engine/Plugins/Experimental/TypedElementsDataStorage/Source/TypedElementsDataStorage/Private/TypedElementsDataStorage.cpp
--- a/Engine/Plugins/Experimental/TypedElementsDataStorage/Source/TypedElementsDataStorage/Private/TypedElementsDataStorage.cpp
+++ b/Engine/Plugins/Experimental/TypedElementsDataStorage/Source/TypedElementsDataStorage/Private/TypedElementsDataStorage.cpp
@@ -17,17 +17,6 @@
 
 #define LOCTEXT_NAMESPACE "FTypedElementsDataStorageModule"
 
-namespace Private
-{
-	static bool GEnableTEDS = false;
-	static FAutoConsoleVariableRef CVarEnableTEDS(
-		TEXT("TEDS.Enable"),
-		GEnableTEDS,
-		TEXT("Enable TypedElementDataStorage"),
-		ECVF_ReadOnly
-		);
-} // namespace Private
-
 // MASS uses CDO in a few places, making it a difficult to consistently register Type Element's Columns and Tags
 // as they may have not been set up to impersonate MASS' Fragments and Tags yet. There are currently no longer
 // any cases where TEDS relies on this but may happen again the future. For the standalone version a static
@@ -62,14 +51,6 @@ void ImpersonateMassTagsAndFragments()
 
 void FTypedElementsDataStorageModule::StartupModule()
 {
-	if (!Private::GEnableTEDS)
-	{
-		UE_LOG(LogTypedElementDataStorage, Log, TEXT("Disabled by TEDS.Enable CVar"));
-		return;
-	}
-
-	UE_LOG(LogTypedElementDataStorage, Log, TEXT("Enabled by TEDS.Enable CVar"));
-
 	// Setup the editor settings;
 	if (ISettingsModule* SettingsModule = FModuleManager::GetModulePtr<ISettingsModule>("Settings"))
 	{
