Fixed a remaining reference to FMassActorFragment from TEDS.

This was causing a crash for some users.

diff --git a/Engine/Plugins/Experimental/TypedElementsDataStorage/Source/TypedElementsDataStorageUI/Private/TypedElementCounterWidget.cpp b/Engine/Plugins/Experimental/TypedElementsDataStorage/Source/TypedElementsDataStorageUI/Private/TypedElementCounterWidget.cpp
--- a/Engine/Plugins/Experimental/TypedElementsDataStorage/Source/TypedElementsDataStorageUI/Private/TypedElementCounterWidget.cpp
+++ b/Engine/Plugins/Experimental/TypedElementsDataStorage/Source/TypedElementsDataStorageUI/Private/TypedElementCounterWidget.cpp
@@ -89,18 +89,6 @@ void UTypedElementCounterWidgetFactory::RegisterWidgetConstructors(ITypedElement
 {
 	using namespace TypedElementQueryBuilder;
 
-	TUniquePtr<FTypedElementCounterWidgetConstructor> ActorCounter = MakeUnique<FTypedElementCounterWidgetConstructor>();
-	ActorCounter->LabelText = LOCTEXT("ActorCounterStatusBarLabel", "{0} {0}|plural(one=Actor, other=Actors)");
-	ActorCounter->ToolTipText = LOCTEXT(
-		"ActorCounterStatusBarToolTip",
-		"The total number of actors currently in the editor, excluding PIE/SIE and previews.");
-	ActorCounter->Query = DataStorage.RegisterQuery(
-		Count().
-		Where().
-			All("/Script/MassActors.MassActorFragment"_Type).
-		Compile());
-	DataStorageUi.RegisterWidgetFactory(WigetPurpose, MoveTemp(ActorCounter));
-
 	TUniquePtr<FTypedElementCounterWidgetConstructor> WidgetCounter = MakeUnique<FTypedElementCounterWidgetConstructor>();
 	WidgetCounter->LabelText = LOCTEXT("WidgetCounterStatusBarLabel", "{0} {0}|plural(one=Widget, other=Widgets)");
 	WidgetCounter->ToolTipText = LOCTEXT(
