Optimization to the defered column removal for a batch of row.

diff --git a/Engine/Plugins/Experimental/TypedElementsDataStorage/Source/TypedElementsDataStorage/Private/Processors/TypedElementProcessorAdaptors.cpp b/Engine/Plugins/Experimental/TypedElementsDataStorage/Source/TypedElementsDataStorage/Private/Processors/TypedElementProcessorAdaptors.cpp
--- a/Engine/Plugins/Experimental/TypedElementsDataStorage/Source/TypedElementsDataStorage/Private/Processors/TypedElementProcessorAdaptors.cpp
+++ b/Engine/Plugins/Experimental/TypedElementsDataStorage/Source/TypedElementsDataStorage/Private/Processors/TypedElementProcessorAdaptors.cpp
@@ -370,12 +370,42 @@ public:
 			{
 				FMassEntityHandle* Entities = RemovedColumns->Entities;
 				int32 Count = RemovedColumns->EntityCount;
+
+				using EntityHandleArray = TArray<FMassEntityHandle, TInlineAllocator<32>>;
+				using EntityArchetypeLookup = TMap<FMassArchetypeHandle, EntityHandleArray, TInlineSetAllocator<32>>;
+				using ArchetypeEntityArray = TArray<FMassArchetypeEntityCollection, TInlineAllocator<32>>;
+
+				// Sort rows (entities) into to matching table (archetype) bucket.
+				EntityArchetypeLookup LookupTable;
 				for (int32 Counter = 0; Counter < Count; ++Counter)
 				{
 					if (System.IsEntityValid(*Entities))
 					{
-						System.RemoveCompositionFromEntity(*Entities++, RemovedColumns->RemoveDescriptor);
+						FMassArchetypeHandle Archetype = System.GetArchetypeForEntity(*Entities);
+						EntityHandleArray& EntityCollection = LookupTable.FindOrAdd(Archetype);
+						EntityCollection.Add(*Entities);
 					}
+					Entities++;
+				}
+
+				// Construct table (archetype) specific row (entity) collections.
+				ArchetypeEntityArray EntityCollections;
+				EntityCollections.Reserve(LookupTable.Num());
+				for (auto It = LookupTable.CreateConstIterator(); It; ++It)
+				{
+					// Could be more effective but the previous implementation was robust when called with duplicate rows.
+					EntityCollections.Emplace(It.Key(), It.Value(), FMassArchetypeEntityCollection::EDuplicatesHandling::FoldDuplicates);
+				}
+
+
+				// This could be improved by adding an operation that would both combine the Fragments and Tags change in one bath operation.
+				if (!RemovedColumns->RemoveDescriptor.Fragments.IsEmpty())
+				{
+					System.BatchChangeFragmentCompositionForEntities(EntityCollections, FMassFragmentBitSet(), RemovedColumns->RemoveDescriptor.Fragments);
+				}
+				if (!RemovedColumns->RemoveDescriptor.Tags.IsEmpty())
+				{
+					System.BatchChangeTagsForEntities(EntityCollections, FMassTagBitSet(), RemovedColumns->RemoveDescriptor.Tags);
 				}
 			});
 	}
