Made TEDS Compat thread-safe and used the new command buffer.

TEDS Compat has been made thread safe and by extension the memento system as well. To offset the limited impact this had on performance and to allow future thread-aware development, this change also moved to using the new command buffer/queue, which resulted in a roughly 30% performance gain in one test scenario.

Added the folowing CVars:
- TEDS.Feature.UseCommandBufferInCompat - When set to true will use the command buffer.
- TEDS.Feature.UseDeferredRemovesInCompat - If true, all removes go through the command buffer. If false, only the removes from the Garbage Collection use the command buffer.
- TEDS.Debug.OptimizeCommandBufferInCompat - Whether or not to apply optimizations to the command buffer.
- TEDS.Debug.PrintCompatCommandBuffer - Prints the command buffer. 0 to disable, 1 to only show the number of no-operation commands or 2 to show nops inline.

Additional changes:
- Updated the Outliner to check entities even if they're already marked as unreachable. They will still fail if the actor proxy has already been deleted. This can be seen when `bUseDeferredRemovesInCompat` is enabled.
- Replaced several cases of TypedElementTableHandle with TypedElementDataStorage::TableHandle.
- Set the InvalidRowHandle value to 0. This works a bit better with Mass as it can early out since the Mass default is 0 as well. It also helps when row handles are used in columns as columns are default zero initialized, meaning this is one less value to explicitly initialize.
- Relocated FTypedElementDatabaseCompatibilityObjectTypeInfo and put it in the new UE::EditorDataStorage namespace.

diff --git a/Engine/Plugins/Experimental/TypedElementsDataStorage/Source/TypedElementsDataStorage/Private/Commands/EditorDataStorageCompatibilityCommands.cpp b/Engine/Plugins/Experimental/TypedElementsDataStorage/Source/TypedElementsDataStorage/Private/Commands/EditorDataStorageCompatibilityCommands.cpp
new file mode 100644
--- /dev/null
+++ b/Engine/Plugins/Experimental/TypedElementsDataStorage/Source/TypedElementsDataStorage/Private/Commands/EditorDataStorageCompatibilityCommands.cpp
@@ -0,0 +1,1138 @@
+// Copyright Epic Games, Inc. All Rights Reserved.
+
+#include "Commands/EditorDataStorageCompatibilityCommands.h"
+
+#include "Algo/BinarySearch.h"
+#include "Elements/Columns/TypedElementCompatibilityColumns.h"
+#include "Elements/Columns/TypedElementMiscColumns.h"
+#include "Elements/Columns/TypedElementTypeInfoColumns.h"
+#include "Elements/Framework/TypedElementIndexHasher.h"
+#include "Elements/Framework/TypedElementQueryBuilder.h"
+#include "Elements/Interfaces/TypedElementDataStorageInterface.h"
+#include "Memento/TypedElementMementoSystem.h"
+#include "TypedElementDatabaseCompatibility.h"
+#include "TypedElementDatabaseEnvironment.h"
+#include "TypedElementDatabaseScratchBuffer.h"
+
+namespace UE::EditorDataStorage
+{
+	UScriptStruct* FAddSyncFromWorldTag::Type = FTypedElementSyncFromWorldTag::StaticStruct();
+	UScriptStruct* FAddInteractiveSyncFromWorldTag::Type = FTypedElementSyncFromWorldInteractiveTag::StaticStruct();
+
+	//
+	// FObjectTypeInfo
+	//
+
+	FObjectTypeInfo::FObjectTypeInfo(const UScriptStruct* InScriptStruct)
+		: TypeInfoType(EObjectType::Struct)
+		, ScriptStruct(InScriptStruct)
+	{}
+
+	FObjectTypeInfo::FObjectTypeInfo(const UClass* InClass)
+		: TypeInfoType(EObjectType::Class)
+		, Class(InClass)
+	{}
+
+	FName FObjectTypeInfo::GetFName() const
+	{
+		switch (TypeInfoType)
+		{
+		case EObjectType::Struct: return ScriptStruct->GetFName();
+		case EObjectType::Class: return Class->GetFName();
+		default: return FName();
+		}
+	}
+
+
+
+	//
+	// FTypeBatchInfoReinstanced
+	//
+
+	template<typename T>
+	FTypeInfoReinstanced* FTypeBatchInfoReinstanced::FindObject(TArrayView<FTypeInfoReinstanced> Range, const TWeakObjectPtr<T>& Object)
+	{
+		// If performance of this suffers too much due to the hash having too many collisions moving to an FObjectKey might be more
+		// more efficient as there will be guaranteed no collisions. However there's currently no way to go from a TWeakObjectPtr
+		// to a TObjectKey<UObject> or FObjectKey.
+
+		uint32 TargetHash = Object.GetWeakPtrTypeHash();
+		int32 FoundIndex = Algo::LowerBoundBy(Range, TargetHash,
+			[](FTypeInfoReinstanced& Reinstanced)
+			{
+				return Reinstanced.Original.GetWeakPtrTypeHash();
+			});
+		// Keep searching linearly since the hash for weak pointers has a high chance of collisions.
+		FTypeInfoReinstanced* Result = Range.GetData() + FoundIndex;
+		FTypeInfoReinstanced* End = Range.end();
+		while (Result < End
+			&& Result->Original.GetWeakPtrTypeHash() == TargetHash
+			&& !Result->Original.HasSameIndexAndSerialNumber(Object))
+		{
+			Result++;
+		}
+		if (Result < End && Result->Original.HasSameIndexAndSerialNumber(Object))
+		{
+			return Result;
+		}
+		return nullptr;
+	}
+
+	template<typename T>
+	TWeakObjectPtr<T> FTypeBatchInfoReinstanced::FindObjectRecursively(
+		TArrayView<FTypeInfoReinstanced> Range, const TWeakObjectPtr<T>& Object)
+	{
+		if (FTypeInfoReinstanced* NewObject = FindObject(Range, Object))
+		{
+			FTypeInfoReinstanced* LastNewObject = NewObject;
+			while (FTypeInfoReinstanced* NextNewObject = FindObject(Range, LastNewObject->Reinstanced))
+			{
+				LastNewObject = NextNewObject;
+			}
+			return Cast<T>(LastNewObject->Reinstanced);
+		}
+		return Object;
+	}
+
+
+
+	//
+	// FPatchData
+	//
+
+	FPatchData::FPatchCommand::FPatchCommand(TArrayView<FTypeInfoReinstanced>& InReinstances)
+		: Reinstances(InReinstances)
+	{
+	}
+
+	void FPatchData::FPatchCommand::operator()(FRegisterTypeTableAssociation& Command)
+	{
+		Command.TypeInfo = FTypeBatchInfoReinstanced::FindObjectRecursively(Reinstances, Command.TypeInfo);
+	}
+
+	void FPatchData::FPatchCommand::operator()(FAddCompatibleExternalObject& Command)
+	{
+		Command.TypeInfo = FTypeBatchInfoReinstanced::FindObjectRecursively(Reinstances, Command.TypeInfo);
+		checkf(Command.TypeInfo.IsValid(), TEXT("A script struct has be re-instanced to an object that's not a script struct."))
+	}
+
+	void FPatchData::FPatchCommand::operator()(FBatchAddCompatibleExternalObject& Command)
+	{
+		TArrayView<TWeakObjectPtr<UScriptStruct>> TypeInfoArray = 
+			TArrayView<TWeakObjectPtr<UScriptStruct>>(Command.TypeInfoArray, Command.Count);
+
+		for (TWeakObjectPtr<UScriptStruct>& TypeInfo : TypeInfoArray)
+		{
+			TypeInfo = FTypeBatchInfoReinstanced::FindObjectRecursively(Reinstances, TypeInfo);
+			checkf(TypeInfo.IsValid(), TEXT("A script struct has be re-instanced to an object that's not a script struct."))
+		}
+	}
+
+
+	bool FPatchData::IsPatchingRequired(const CompatibilityCommandBuffer::FCollection& Commands)
+	{
+		return Commands.GetCommandCount<FTypeInfoReinstanced>() > 0;
+	}
+
+	void FPatchData::RunPatch(CompatibilityCommandBuffer::FCollection& Commands, UTypedElementDatabaseCompatibility& StorageCompat,
+		FTypedElementDatabaseScratchBuffer& ScratchBuffer)
+	{
+		using namespace TypedElementDataStorage;
+
+		uint32 ReinstanceCount = Commands.GetCommandCount<FTypeInfoReinstanced>();
+		FTypeInfoReinstanced* ReinstanceArray = ScratchBuffer.EmplaceArray<FTypeInfoReinstanced>(ReinstanceCount);
+		TArrayView<FTypeInfoReinstanced> ReinstanceView(ReinstanceArray, ReinstanceCount);
+
+		// Populate the list of re-instance data and disable the commands.
+		FTypeInfoReinstanced* ReinstanceArrayIt = ReinstanceArray;
+		Commands.ForEach(
+			[&Commands, &ReinstanceArrayIt]
+			(int32 Index, CompatibilityCommandBuffer::TCommandVariant& Command)
+			{
+				if (FTypeInfoReinstanced* ReinstanceData = Command.TryGet<FTypeInfoReinstanced>())
+				{
+					*ReinstanceArrayIt = *ReinstanceData;
+					Commands.ReplaceCommand<FNopCommand>(Index);
+					++ReinstanceArrayIt;
+				}
+			});
+
+		// Add a command to process type information stored in TEDS itself in a later step.
+		Commands.AddCommand<FTypeBatchInfoReinstanced>(FTypeBatchInfoReinstanced{ .Batch = ReinstanceView });
+
+		// Sort the extracted array by the index of the original value for faster indexing in the future.
+		ReinstanceView.Sort([](const FTypeInfoReinstanced& Lhs, const FTypeInfoReinstanced& Rhs)
+			{
+				return Lhs.Original.GetWeakPtrTypeHash() < Rhs.Original.GetWeakPtrTypeHash();
+			});
+
+		// Patch the type info table.
+		// There can be a large number of re-instanced type information but typically a limited number of type-to-table mappings. Instead
+		// of searching through all the type information updates, copy the type-to-table locally and rebuild it with update info.
+		using TypeTablePair = TPair<TWeakObjectPtr<UStruct>, TableHandle>;
+		TypeTablePair* TempTypeToTable = ScratchBuffer.EmplaceArray<TypeTablePair>(StorageCompat.TypeToTableMap.Num());
+		TArrayView<TypeTablePair> TempTypeToTableView(TempTypeToTable, StorageCompat.TypeToTableMap.Num());
+		TypeTablePair* TempTypeToTableIt = TempTypeToTable;
+		for (decltype(StorageCompat.TypeToTableMap)::TIterator It = StorageCompat.TypeToTableMap.CreateIterator(); It; ++It)
+		{
+			TempTypeToTableIt->Key = It.Key();
+			TempTypeToTableIt->Value = It.Value();
+			++TempTypeToTableIt;
+		}
+		StorageCompat.TypeToTableMap.Reset();
+		for (TypeTablePair It : TempTypeToTableView)
+		{
+			StorageCompat.TypeToTableMap.Add(FTypeBatchInfoReinstanced::FindObjectRecursively(ReinstanceView, It.Key), It.Value);
+		}
+
+		// Patch locally cached type information
+		UScriptStruct** LocalTypeInfoStorage[] =
+		{ 
+			&FAddSyncFromWorldTag::Type,
+			&FAddInteractiveSyncFromWorldTag::Type
+		};
+		TArrayView<UScriptStruct**> LocalTypeInfo(LocalTypeInfoStorage, sizeof(LocalTypeInfoStorage) / sizeof(UScriptStruct**));
+		for (UScriptStruct** TypeInfo : LocalTypeInfo)
+		{
+			*TypeInfo = Cast<UScriptStruct>(
+				FTypeBatchInfoReinstanced::FindObjectRecursively(ReinstanceView, TWeakObjectPtr<UStruct>(*TypeInfo)));
+		}
+
+		// Patch existing commands with the new type information.
+		Commands.Process(FPatchCommand(ReinstanceView));
+	}
+	
+
+
+	//
+	// FPrepareCommands
+	//
+
+	FPrepareCommands::FPrepareCommands(ITypedElementDataStorageInterface& InStorage, UTypedElementDatabaseCompatibility& InStorageCompat,
+		CompatibilityCommandBuffer::FCollection& InCommands)
+		: Storage(InStorage)
+		, StorageCompat(InStorageCompat)
+		, Commands(InCommands)
+	{}
+
+	void FPrepareCommands::operator()(FAddCompatibleUObject& Object)
+	{
+		using namespace TypedElementDataStorage;
+
+		UObject* ObjectPtr = Object.Object.Get();
+		if (ObjectPtr && Storage.IsRowAvailable(Object.Row))
+		{
+			Object.Table = StorageCompat.FindBestMatchingTable(ObjectPtr->GetClass());
+			checkf(Object.Table != InvalidTableHandle,
+				TEXT("The data storage could not find any matching tables for object of type '%s'. "
+					"This can mean that the object doesn't derive from UObject or that a table for UObject is no longer registered."),
+				*Object.Object->GetClass()->GetFName().ToString());
+		}
+		else
+		{
+			Commands.ReplaceCommand<FNopCommand>(CurrentIndex);
+		}
+	}
+
+	void FPrepareCommands::operator()(FAddCompatibleExternalObject& Object)
+	{
+		using namespace TypedElementDataStorage;
+
+		if (Storage.IsRowAvailable(Object.Row) && Object.Object != nullptr)
+		{
+			Object.Table = StorageCompat.FindBestMatchingTable(Object.TypeInfo.Get());
+			Object.Table = (Object.Table != InvalidTableHandle) ? Object.Table : StorageCompat.StandardExternalObjectTable;
+		}
+		else
+		{
+			Commands.ReplaceCommand<FNopCommand>(CurrentIndex);
+		}
+	}
+
+	void FPrepareCommands::operator()(FRemoveCompatibleUObject& Command)
+	{
+		using namespace TypedElementDataStorage;
+
+		if (Command.ObjectRow == InvalidRowHandle)
+		{
+			IndexHash Hash = GenerateIndexHash(Command.Object);
+			Command.ObjectRow = Storage.FindIndexedRow(Hash);
+			if (!Storage.IsRowAvailable(Command.ObjectRow))
+			{
+				Commands.ReplaceCommand<FNopCommand>(CurrentIndex);
+			}
+		}
+	}
+
+	void FPrepareCommands::operator()(FRemoveCompatibleExternalObject& Command)
+	{
+		using namespace TypedElementDataStorage;
+
+		IndexHash Hash = GenerateIndexHash(Command.Object);
+		Command.ObjectRow = Storage.FindIndexedRow(Hash);
+		if (!Storage.IsRowAvailable(Command.ObjectRow))
+		{
+			Commands.ReplaceCommand<FNopCommand>(CurrentIndex);
+		}
+	}
+
+	void FPrepareCommands::operator()(FAddSyncFromWorldTag& Command)
+	{
+		using namespace TypedElementDataStorage;
+
+		Command.Row = StorageCompat.FindRowWithCompatibleObject(Command.Target);
+		if (!Storage.IsRowAvailable(Command.Row))
+		{
+			Commands.ReplaceCommand<FNopCommand>(CurrentIndex);
+		}
+	}
+
+	void FPrepareCommands::operator()(FRemoveInteractiveSyncFromWorldTag& Command)
+	{
+		using namespace TypedElementDataStorage;
+
+		Command.Row = StorageCompat.FindRowWithCompatibleObject(Command.Target);
+		if (!Storage.IsRowAvailable(Command.Row))
+		{
+			Commands.ReplaceCommand<FNopCommand>(CurrentIndex);
+		}
+	}
+
+	void FPrepareCommands::RunPreparation(ITypedElementDataStorageInterface& Storage, UTypedElementDatabaseCompatibility& StorageCompat,
+		CompatibilityCommandBuffer::FCollection& Commands)
+	{
+		FPrepareCommands PrepareVisitor(Storage, StorageCompat, Commands);
+		Commands.ForEach(
+			[&PrepareVisitor](int32 Index, CompatibilityCommandBuffer::TCommandVariant& Command)
+			{
+				PrepareVisitor.CurrentIndex = Index;
+				Visit(PrepareVisitor, Command);
+			});
+	}
+
+	void FPrepareCommands::operator()(FAddInteractiveSyncFromWorldTag& Command)
+	{
+		using namespace TypedElementDataStorage;
+
+		Command.Row = StorageCompat.FindRowWithCompatibleObject(Command.Target);
+		if (!Storage.IsRowAvailable(Command.Row))
+		{
+			Commands.ReplaceCommand<FNopCommand>(CurrentIndex);
+		}
+	}
+
+	//
+	// FGetSourceRowHandle
+	//
+
+	TypedElementDataStorage::RowHandle FGetSourceRowHandle::operator()(const FAddCompatibleUObject& Command)
+	{
+		return Command.Row;
+	}
+
+	TypedElementDataStorage::RowHandle FGetSourceRowHandle::operator()(const FAddCompatibleExternalObject& Command)
+	{
+		return Command.Row;
+	}
+
+	TypedElementDataStorage::RowHandle FGetSourceRowHandle::operator()(const FCreateMemento& Command)
+	{
+		return Command.TargetRow;
+	}
+
+	TypedElementDataStorage::RowHandle FGetSourceRowHandle::operator()(const FRestoreMemento& Command)
+	{
+		return Command.TargetRow;
+	}
+
+	TypedElementDataStorage::RowHandle FGetSourceRowHandle::operator()(const FDestroyMemento& Command)
+	{
+		return Command.MementoRow;
+	}
+
+	TypedElementDataStorage::RowHandle FGetSourceRowHandle::operator()(const FRemoveCompatibleUObject& Command)
+	{
+		return Command.ObjectRow;
+	}
+
+	TypedElementDataStorage::RowHandle FGetSourceRowHandle::operator()(const FRemoveCompatibleExternalObject& Command)
+	{
+		return Command.ObjectRow;
+	}
+
+	TypedElementDataStorage::RowHandle FGetSourceRowHandle::Get(const CompatibilityCommandBuffer::TCommandVariant& Command)
+	{
+		return Visit(FGetSourceRowHandle{}, Command);
+	}
+
+	TypedElementDataStorage::RowHandle FGetSourceRowHandle::operator()(const FAddInteractiveSyncFromWorldTag& Command)
+	{
+		return Command.Row;
+	}
+
+	TypedElementDataStorage::RowHandle FGetSourceRowHandle::operator()(const FRemoveInteractiveSyncFromWorldTag& Command)
+	{
+		return Command.Row;
+	}
+
+	TypedElementDataStorage::RowHandle FGetSourceRowHandle::operator()(const FAddSyncFromWorldTag& Command)
+	{
+		return Command.Row;
+	}
+
+
+
+	//
+	// FGetCommandGroupId
+	//
+
+	// FAddCompatibleUObject group begin: The following commands needs to be sorted together with FAddCompatibleUObject.
+
+	SIZE_T FGetCommandGroupId::operator()(const FBatchAddCompatibleUObject& Command)
+	{
+		return CompatibilityCommandBuffer::TCommandVariant::template IndexOfType<FAddCompatibleUObject>();
+	}
+
+	SIZE_T FGetCommandGroupId::operator()(const FAddCompatibleExternalObject& Command)
+	{
+		return CompatibilityCommandBuffer::TCommandVariant::template IndexOfType<FAddCompatibleUObject>();
+	}
+
+	SIZE_T FGetCommandGroupId::operator()(const FBatchAddCompatibleExternalObject& Command)
+	{
+		return CompatibilityCommandBuffer::TCommandVariant::template IndexOfType<FAddCompatibleUObject>();
+	}
+
+	SIZE_T FGetCommandGroupId::operator()(const FCreateMemento& Command)
+	{
+		return CompatibilityCommandBuffer::TCommandVariant::template IndexOfType<FAddCompatibleUObject>();
+	}
+
+	SIZE_T FGetCommandGroupId::operator()(const FRestoreMemento& Command)
+	{
+		return CompatibilityCommandBuffer::TCommandVariant::template IndexOfType<FAddCompatibleUObject>();
+	}
+
+	SIZE_T FGetCommandGroupId::operator()(const FDestroyMemento& Command)
+	{
+		return CompatibilityCommandBuffer::TCommandVariant::template IndexOfType<FAddCompatibleUObject>();
+	}
+
+	SIZE_T FGetCommandGroupId::operator()(const FRemoveCompatibleUObject& Command)
+	{
+		return CompatibilityCommandBuffer::TCommandVariant::template IndexOfType<FAddCompatibleUObject>();
+	}
+
+	SIZE_T FGetCommandGroupId::operator()(const FRemoveCompatibleExternalObject& Command)
+	{
+		return CompatibilityCommandBuffer::TCommandVariant::template IndexOfType<FAddCompatibleUObject>();
+	}
+
+	// FAddCompatibleUObject group end
+	
+	// FAddInteractiveSyncFromWorldTag group begin: The following commands needs to be sorted together with FAddInteractiveSyncFromWorldTag.
+
+	SIZE_T FGetCommandGroupId::operator()(const FRemoveInteractiveSyncFromWorldTag& Command)
+	{
+		return CompatibilityCommandBuffer::TCommandVariant::template IndexOfType<FAddInteractiveSyncFromWorldTag>();
+	}
+
+	// FAddInteractiveSyncFromWorldTag group end
+
+	SIZE_T FGetCommandGroupId::Get(const CompatibilityCommandBuffer::TCommandVariant& Command)
+	{
+		return Visit(FGetCommandGroupId{}, Command);
+	}
+
+
+
+
+
+	//
+	// FSorter
+	//
+
+	void FSorter::SortCommands(CompatibilityCommandBuffer::FCollection& Commands)
+	{
+		using namespace TypedElementDataStorage;
+
+		// First sort commands by the row they operate on, using stable sort so commands remain in the same range as they get processed.
+		// Within a row, commands that can't be reordered relative to each other are grouped together while the remaining commands are
+		// ordered in the order they were declared in the command buffer.
+		Commands.Sort<true /* Use stable sort */>(
+			[](const CompatibilityCommandBuffer::TCommandVariant& Lhs, const CompatibilityCommandBuffer::TCommandVariant& Rhs)
+			{
+				RowHandle Left = FGetSourceRowHandle::Get(Lhs);
+				RowHandle Right = FGetSourceRowHandle::Get(Rhs);
+				
+				// Sort commands by index to make sure operations stay close together so the same table gets accessed more
+				// frequently, but operations continue to execute in the order they were issued. For instance an Add + Remove
+				// gives a different result from a Remove + Add.
+				if (Left < Right)
+				{
+					return true;
+				}
+				else if (Left == Right)
+				{
+					SIZE_T LeftGroupId = FGetCommandGroupId::Get(Lhs);
+					SIZE_T RightGroupId = FGetCommandGroupId::Get(Rhs);
+					return LeftGroupId < RightGroupId;
+				}
+				return false;
+			});
+	}
+
+
+
+	//
+	// FCommandProcessor
+	//
+
+	FCommandProcessor::FCommandProcessor(ITypedElementDataStorageInterface& InStorage, UTypedElementDatabaseCompatibility& InStorageCompatibility)
+		: Storage(InStorage)
+		, StorageCompatibility(InStorageCompatibility)
+		, MementoSystem(InStorageCompatibility.Environment->GetMementoSystem())
+	{}
+
+	void FCommandProcessor::SetupRow(TypedElementDataStorage::RowHandle Row, UObject* Object)
+	{
+		Storage.AddColumn(Row, FTypedElementUObjectColumn{ .Object = Object });
+		Storage.AddColumn(Row, FTypedElementUObjectIdColumn
+			{
+				.Id = Object->GetUniqueID(),
+				.SerialNumber = GUObjectArray.GetSerialNumber(Object->GetUniqueID())
+			});
+		Storage.AddColumn(Row, FTypedElementClassTypeInfoColumn{ .TypeInfo = Object->GetClass() });
+		if (Object->HasAnyFlags(RF_ClassDefaultObject))
+		{
+			Storage.AddColumn<FTypedElementClassDefaultObjectTag>(Row);
+		}
+		// Make sure the new row is tagged for update.
+		Storage.AddColumn<FTypedElementSyncFromWorldTag>(Row);
+		StorageCompatibility.TriggerOnObjectAdded(Object, Object->GetClass(), Row);
+	}
+
+	void FCommandProcessor::SetupRow(TypedElementDataStorage::RowHandle Row, void* Object, TWeakObjectPtr<UScriptStruct> TypeInfo)
+	{
+		Storage.AddColumn(Row, FTypedElementExternalObjectColumn{ .Object = Object });
+		Storage.AddColumn(Row, FTypedElementScriptStructTypeInfoColumn{ .TypeInfo = TypeInfo });
+		// Make sure the new row is tagged for update.
+		Storage.AddColumn<FTypedElementSyncFromWorldTag>(Row);
+		StorageCompatibility.TriggerOnObjectAdded(Object, TypeInfo.Get(), Row);
+	}
+
+	void FCommandProcessor::operator()(FTypeBatchInfoReinstanced& Command)
+	{
+		using namespace TypedElementDataStorage;
+		using namespace TypedElementQueryBuilder;
+
+		StorageCompatibility.Storage->RunQuery(StorageCompatibility.ClassTypeInfoQuery, CreateDirectQueryCallbackBinding(
+			[Range = Command.Batch](IDirectQueryContext& Context, FTypedElementClassTypeInfoColumn& Type)
+			{
+				Type.TypeInfo = FTypeBatchInfoReinstanced::FindObjectRecursively(Range, Type.TypeInfo);
+				checkf(Type.TypeInfo.IsValid(),
+					TEXT("Type info column in data storage has been re-instanced to an object without class type information"));
+			}));
+
+		StorageCompatibility.Storage->RunQuery(StorageCompatibility.ScriptStructTypeInfoQuery, CreateDirectQueryCallbackBinding(
+			[Range = Command.Batch](IDirectQueryContext& Context, FTypedElementScriptStructTypeInfoColumn& Type)
+			{
+				Type.TypeInfo = FTypeBatchInfoReinstanced::FindObjectRecursively(Range, Type.TypeInfo);
+				checkf(Type.TypeInfo.IsValid(),
+					TEXT("Type info column in data storage has been re-instanced to an object without struct type information"));
+			}));
+	}
+
+	void FCommandProcessor::operator()(FRegisterTypeTableAssociation& Command)
+	{
+		StorageCompatibility.TypeToTableMap.Add(Command.TypeInfo, Command.Table);
+	}
+
+	void FCommandProcessor::operator()(FRegisterObjectAddedCallback& Command)
+	{
+		StorageCompatibility.ObjectAddedCallbackList.Emplace(MoveTemp(Command.Callback), Command.Handle);
+	}
+
+	void FCommandProcessor::operator()(FUnregisterObjectAddedCallback& Command)
+	{
+		StorageCompatibility.ObjectAddedCallbackList.RemoveAll(
+			[Handle = Command.Handle](const TPair<ObjectAddedCallback, FDelegateHandle>& Element) -> bool
+			{
+				return Element.Value == Handle;
+			});
+	}
+
+	void FCommandProcessor::operator()(FAddCompatibleUObject& Object)
+	{
+		UObject* ObjectPtr = Object.Object.Get();
+		checkf(ObjectPtr, TEXT(
+			"Expected a valid object pointer. If there isn't one here then the filter pass did not correctly clean up this command."));
+		Storage.AddRow(Object.Row, Object.Table, 
+			[ObjectPtr, this](TypedElementDataStorage::RowHandle Row)
+			{
+				SetupRow(Row, ObjectPtr);
+			});
+	}
+
+	void FCommandProcessor::operator()(FBatchAddCompatibleUObject& Batch)
+	{
+		using namespace TypedElementDataStorage;
+
+		TConstArrayView<RowHandle> Rows(Batch.RowArray, Batch.Count);
+		RowHandle* RowIterator = Batch.RowArray;
+		TWeakObjectPtr<UObject>* ObjectIterator = Batch.ObjectArray;
+
+		Storage.BatchAddRow(Batch.Table, Rows, [&RowIterator, &ObjectIterator, this](RowHandle Row)
+			{
+				checkf(Row == *RowIterator++, TEXT("Expecting the same sequence of rows when batch adding object to TEDS Compatibility."));
+				UObject* ObjectPtr = ObjectIterator->Get();
+				checkf(ObjectPtr, TEXT(
+					"Expected a valid object pointer. If there isn't one here then the filter pass did not correctly clean up this command."));
+				SetupRow(Row, ObjectPtr);
+				ObjectIterator++;
+			});
+	}
+
+	void FCommandProcessor::operator()(FAddCompatibleExternalObject& Object)
+	{
+		Storage.AddRow(Object.Row, Object.Table, [&Object, this](TypedElementDataStorage::RowHandle Row)
+			{
+				SetupRow(Object.Row, Object.Object, Object.TypeInfo);
+			});
+	}
+
+	void FCommandProcessor::operator()(FBatchAddCompatibleExternalObject& Batch)
+	{
+		using namespace TypedElementDataStorage;
+
+		TConstArrayView<RowHandle> Rows(Batch.RowArray, Batch.Count);
+		RowHandle* RowIterator = Batch.RowArray;
+		void** ObjectIterator = Batch.ObjectArray;
+		TWeakObjectPtr<UScriptStruct>* TypeInfoIterator = Batch.TypeInfoArray;
+
+		Storage.BatchAddRow(Batch.Table, Rows, [&RowIterator, &ObjectIterator, &TypeInfoIterator, this](RowHandle Row)
+			{
+				checkf(Row == *RowIterator++, TEXT("Expecting the same sequence of rows when batch adding object to TEDS Compatibility."));
+				SetupRow(Row, *ObjectIterator++, *TypeInfoIterator++);
+			});
+	}
+
+	void FCommandProcessor::operator()(FCreateMemento& Command)
+	{
+		MementoSystem.CreateMemento(Command.ReservedMementoRow, Command.TargetRow);
+	}
+
+	void FCommandProcessor::operator()(FRestoreMemento& Command)
+	{
+		MementoSystem.RestoreMemento(Command.MementoRow, Command.TargetRow);
+	}
+
+	void FCommandProcessor::operator()(FDestroyMemento& Command)
+	{
+		MementoSystem.DestroyMemento(Command.MementoRow);
+	}
+
+	void FCommandProcessor::operator()(FRemoveCompatibleUObject& Command)
+	{
+		if (Storage.IsRowAssigned(Command.ObjectRow))
+		{
+			const FTypedElementClassTypeInfoColumn* TypeInfoColumn = Storage.GetColumn<FTypedElementClassTypeInfoColumn>(Command.ObjectRow);
+			checkf(TypeInfoColumn, 
+				TEXT("Missing type information for removed UObject at ptr 0x%p [%s]"), Command.Object, *Command.Object->GetName());
+			StorageCompatibility.TriggerOnPreObjectRemoved(Command.Object, TypeInfoColumn->TypeInfo.Get(), Command.ObjectRow);
+		}
+		Storage.RemoveRow(Command.ObjectRow);
+	}
+
+	void FCommandProcessor::operator()(FRemoveCompatibleExternalObject& Command)
+	{
+		if (Storage.IsRowAssigned(Command.ObjectRow))
+		{
+			const FTypedElementScriptStructTypeInfoColumn* TypeInfoColumn = Storage.GetColumn<FTypedElementScriptStructTypeInfoColumn>(Command.ObjectRow);
+			checkf(TypeInfoColumn, TEXT("Missing type information for removed void* object at ptr 0x%p"), Command.Object);
+			StorageCompatibility.TriggerOnPreObjectRemoved(Command.Object, TypeInfoColumn->TypeInfo.Get(), Command.ObjectRow);
+		}
+		Storage.RemoveRow(Command.ObjectRow);
+	}
+
+	void FCommandProcessor::operator()(FAddInteractiveSyncFromWorldTag& Command)
+	{
+		Storage.AddColumn(Command.Row, FAddInteractiveSyncFromWorldTag::Type);
+	}
+
+	void FCommandProcessor::operator()(FRemoveInteractiveSyncFromWorldTag& Command)
+	{
+		Storage.AddRemoveColumns(Command.Row,
+			{ FAddSyncFromWorldTag::Type },
+			{ FAddInteractiveSyncFromWorldTag::Type });
+	}
+
+	void FCommandProcessor::operator()(FAddSyncFromWorldTag& Command)
+	{
+		Storage.AddColumn(Command.Row, FAddSyncFromWorldTag::Type);
+	}
+
+
+
+
+	//
+	// FRecordCommands
+	//
+
+	void FRecordCommands::operator()(const FNopCommand& Command) 
+	{
+		if (bIncludeNops)
+		{
+			CommandDescriptions.Append(TEXT("    FNopCommand\n"));
+		}
+	}
+
+	void FRecordCommands::operator()(const FTypeInfoReinstanced& Command)
+	{
+		CommandDescriptions.Append(TEXT("    FTypeInfoReinstanced: '"));
+		Command.Reinstanced.GetEvenIfUnreachable()->AppendName(CommandDescriptions);
+		CommandDescriptions.Append(TEXT("' \n"));
+	}
+
+	void FRecordCommands::operator()(const FTypeBatchInfoReinstanced& Command)
+	{
+		CommandDescriptions.Appendf(TEXT("    FTypeBatchInfoReinstanced: %i re-instances\n"), Command.Batch.Num());
+		for (FTypeInfoReinstanced& TypeInfo : Command.Batch)
+		{
+			CommandDescriptions.Append(TEXT("        "));
+			TypeInfo.Reinstanced.GetEvenIfUnreachable()->AppendName(CommandDescriptions);
+			CommandDescriptions.Append(TEXT("\n"));
+		}
+	}
+
+	void FRecordCommands::operator()(const FRegisterTypeTableAssociation& Command)
+	{
+		CommandDescriptions.Append(TEXT("    FRegisterTypeTableAssociation: "));
+		Command.TypeInfo->AppendName(CommandDescriptions);
+		CommandDescriptions.Append(TEXT("\n"));
+	}
+
+	void FRecordCommands::operator()(const FRegisterObjectAddedCallback& Command)
+	{
+		CommandDescriptions.Append(TEXT("    FRegisterObjectAddedCallback\n"));
+	}
+
+	void FRecordCommands::operator()(const FUnregisterObjectAddedCallback& Command)
+	{
+		CommandDescriptions.Append(TEXT("    FUnregisterObjectAddedCallback\n"));
+	}
+
+	void FRecordCommands::operator()(const FAddCompatibleUObject& Command)
+	{
+		CommandDescriptions.Append(TEXT("    FAddCompatibleUObject: '"));
+		Command.Object->AppendName(CommandDescriptions);
+		CommandDescriptions.Appendf(TEXT("' row %llu, table %llu\n"), Command.Row, Command.Table);
+	}
+
+	void FRecordCommands::operator()(const FBatchAddCompatibleUObject& Command)
+	{
+		CommandDescriptions.Appendf(TEXT("    FBatchAddCompatibleUObject: %u objects, table %llu\n"), Command.Count, Command.Table);
+		for (uint32 Index = 0; Index < Command.Count; ++Index)
+		{
+			CommandDescriptions.Append(TEXT("        '"));
+			Command.ObjectArray[Index]->AppendName(CommandDescriptions);
+			CommandDescriptions.Appendf(TEXT("`row %llu\n"), Command.RowArray[Index]);
+		}
+	}
+
+	void FRecordCommands::operator()(const FAddCompatibleExternalObject& Command)
+	{
+		CommandDescriptions.Append(TEXT("    FAddCompatibleExternalObject: '"));
+		Command.TypeInfo->AppendName(CommandDescriptions);
+		CommandDescriptions.Appendf(TEXT("' row %llu, table %llu\n"), Command.Row, Command.Table);
+	}
+
+	void FRecordCommands::operator()(const FBatchAddCompatibleExternalObject& Command)
+	{
+		CommandDescriptions.Appendf(TEXT("    FBatchAddCompatibleExternalObject: %u objects, table %llu\n"), Command.Count, Command.Table);
+	}
+
+	void FRecordCommands::operator()(const FCreateMemento& Command)
+	{
+		CommandDescriptions.Appendf(TEXT("    FCreateMemento: row %llu, memento row %llu\n"), Command.TargetRow, Command.ReservedMementoRow);
+	}
+
+	void FRecordCommands::operator()(const FRestoreMemento& Command)
+	{
+		CommandDescriptions.Appendf(TEXT("    FRestoreMemento: row %llu, memento row %llu\n"), Command.TargetRow, Command.MementoRow);
+	}
+
+	void FRecordCommands::operator()(const FDestroyMemento& Command)
+	{
+		CommandDescriptions.Appendf(TEXT("    FDestroyMemento: memento row %llu\n"), Command.MementoRow);
+	}
+
+	void FRecordCommands::operator()(const FRemoveCompatibleUObject& Command)
+	{
+		CommandDescriptions.Append(TEXT("    FRemoveCompatibleUObject: '"));
+		Command.Object->AppendName(CommandDescriptions);
+		CommandDescriptions.Appendf(TEXT("' row %llu\n"), Command.ObjectRow);
+	}
+
+	void FRecordCommands::operator()(const FRemoveCompatibleExternalObject& Command)
+	{
+		CommandDescriptions.Appendf(TEXT("    FRemoveCompatibleExternalObject: row %llu\n"), Command.ObjectRow);
+	}
+
+	void FRecordCommands::operator()(const FAddInteractiveSyncFromWorldTag& Command)
+	{
+		CommandDescriptions.Appendf(TEXT("    FAddInteractiveSyncFromWorldTag: row %llu\n"), Command.Row);
+	}
+
+	void FRecordCommands::operator()(const FRemoveInteractiveSyncFromWorldTag& Command)
+	{
+		CommandDescriptions.Appendf(TEXT("    FRemoveInteractiveSyncFromWorldTag: row %llu\n"), Command.Row);
+	}
+
+	void FRecordCommands::operator()(const FAddSyncFromWorldTag& Command)
+	{
+		CommandDescriptions.Appendf(TEXT("    FAddSyncFromWorldTag: row %llu\n"), Command.Row);
+	}
+
+	FString FRecordCommands::PrintToString(CompatibilityCommandBuffer::FCollection& Commands, bool bIncludeNops)
+	{
+		FRecordCommands Printer;
+		if (bIncludeNops)
+		{
+			Printer.bIncludeNops = true;
+			Commands.Process(Printer);
+		}
+		else
+		{
+			Printer.bIncludeNops = false;
+			Commands.Process(Printer);
+		}
+		return MoveTemp(Printer.CommandDescriptions);
+	}
+
+
+
+	
+	//
+	// FCommandOptimizer
+	//
+
+	FCommandOptimizer::FCommandOptimizer(
+		CompatibilityCommandBuffer::FOptimizer& InOptimizer, FTypedElementDatabaseScratchBuffer& InScratchBuffer)
+		: Optimizer(InOptimizer)
+		, ScratchBuffer(InScratchBuffer)
+	{}
+
+	void FCommandOptimizer::operator()(const FAddCompatibleUObject& Command)
+	{
+		using namespace TypedElementDataStorage;
+
+		int32 Count = FoldCommandsForAdd(Command);
+
+		// If there are more than 1 adds, batch them together into a batch call and nop the additional adds out.
+		if (Count > 1)
+		{
+			RowHandle* Rows = ScratchBuffer.EmplaceArray<RowHandle>(Count);
+			TWeakObjectPtr<UObject>* Objects = ScratchBuffer.EmplaceArray<TWeakObjectPtr<UObject>>(Count);
+
+			Rows[0] = Command.Row;
+			Objects[0] = Command.Object;
+
+			for (int32 Index = 1; Index < Count;)
+			{
+				if (FAddCompatibleUObject* Right = Optimizer.GetRight().TryGet<FAddCompatibleUObject>())
+				{
+					Rows[Index] = Right->Row;
+					Objects[Index] = Right->Object;
+					Optimizer.ReplaceRight<FNopCommand>();
+					++Index;
+				}
+				Optimizer.MoveToNextRight();
+			}
+			Optimizer.ReplaceLeft<FBatchAddCompatibleUObject>(
+				FBatchAddCompatibleUObject
+				{
+					.ObjectArray = Objects,
+					.RowArray = Rows,
+					.Table = Command.Table,
+					.Count = static_cast<uint32>(Count)
+				});
+			// Skip over the nops for the next optimization.
+			Optimizer.MoveLeftToRight();
+		}
+	}
+	
+	void FCommandOptimizer::operator()(const FAddCompatibleExternalObject& Command)
+	{
+		using namespace TypedElementDataStorage;
+
+		int32 Count = FoldCommandsForAdd(Command);
+
+		// If there are more than 1 adds, batch them together into a batch call and nop them out.
+		if (Count > 1)
+		{
+			RowHandle* Rows = ScratchBuffer.EmplaceArray<RowHandle>(Count);
+			void** Objects = ScratchBuffer.EmplaceArray<void*>(Count);
+			TWeakObjectPtr<UScriptStruct>* TypeInfo = ScratchBuffer.EmplaceArray<TWeakObjectPtr<UScriptStruct>>(Count);
+
+			Rows[0] = Command.Row;
+			Objects[0] = Command.Object;
+			TypeInfo[0] = Command.TypeInfo;
+
+			for (int32 Index = 1; Index < Count; ++Index)
+			{
+				FAddCompatibleExternalObject& Right = Optimizer.GetRight().Get<FAddCompatibleExternalObject>();
+				Rows[Index] = Right.Row;
+				Objects[Index] = Right.Object;
+				TypeInfo[Index] = Right.TypeInfo;
+				Optimizer.ReplaceRight<FNopCommand>();
+				Optimizer.MoveToNextRight();
+			}
+			Optimizer.GetLeft().Emplace<FBatchAddCompatibleExternalObject>(
+				FBatchAddCompatibleExternalObject
+				{
+					.ObjectArray = Objects,
+					.TypeInfoArray = TypeInfo,
+					.RowArray = Rows,
+					.Table = Command.Table,
+					.Count = static_cast<uint32>(Count)
+				});
+			// Skip over the nops for the next optimization.
+			Optimizer.MoveLeftToRight();
+		}
+	}
+
+	void FCommandOptimizer::operator()(const FRemoveCompatibleUObject& Command)
+	{
+		using namespace TypedElementDataStorage;
+
+		// Remove everything after the delete because no operation on a row is going to succeed after
+		// the row has been deleted.
+		while(Optimizer.IsValid())
+		{
+			RowHandle RightRow = FGetSourceRowHandle::Get(Optimizer.GetRight());
+			if (RightRow == Command.ObjectRow)
+			{
+				Optimizer.ReplaceRight<FNopCommand>();
+				Optimizer.MoveToNextRight();
+			}
+			else
+			{
+				break;
+			}
+		}
+		Optimizer.MoveLeftBeforeRight();
+	}
+
+	void FCommandOptimizer::operator()(const FRemoveCompatibleExternalObject& Command)
+	{
+		using namespace TypedElementDataStorage;
+		
+		// Remove everything after the delete because no operation on a row is going to succeed after
+		// the row has been deleted.
+		while (Optimizer.IsValid())
+		{
+			RowHandle RightRow = FGetSourceRowHandle::Get(Optimizer.GetRight());
+			if (RightRow == Command.ObjectRow)
+			{
+				Optimizer.ReplaceRight<FNopCommand>();
+				Optimizer.MoveToNextRight();
+			}
+			else
+			{
+				break;
+			}
+		}
+		Optimizer.MoveLeftBeforeRight();
+	}
+
+	template<typename AddCommandType>
+	int32 FCommandOptimizer::FoldCommandsForAdd(const AddCommandType& Command)
+	{
+		using namespace TypedElementDataStorage;
+
+		// Discover the longest chain of adds for the same table and remove any async tag additions as it'll be included
+		// in the table for UObjects.
+		int32 Count = 1;
+		RowHandle SourceRow = Command.Row;
+		for (; Optimizer.IsValid(); Optimizer.MoveToNextRight())
+		{
+			const CompatibilityCommandBuffer::TCommandVariant& TargetCommand = Optimizer.GetRight();
+			if (const AddCommandType* Right = TargetCommand.TryGet<AddCommandType>())
+			{
+				if (Right->Table == Command.Table)
+				{
+					++Count;
+					SourceRow = Right->Row;
+				}
+				else
+				{
+					// No longer in the same table, so stop.
+					break;
+				}
+			}
+			else
+			{
+				RowHandle TargetRow = FGetSourceRowHandle::Get(TargetCommand);
+				if (TargetRow == SourceRow)
+				{
+					if (TargetCommand.IsType<FAddSyncFromWorldTag>())
+					{
+						// These are always ordered after the interactive tags, so if there are any interactive tags these
+						// are already folded into a single one. If there are not then this pass will fold them into a single
+						// one.
+						Optimizer.ReplaceRight<FNopCommand>();
+					}
+					else if (TargetCommand.IsType<FAddInteractiveSyncFromWorldTag>() ||
+						TargetCommand.IsType<FRemoveInteractiveSyncFromWorldTag>())
+					{
+						// Run optimizations on the interactive tags so they can get folded and avoid complex checks here.
+						RunRightOnRowCluster(SourceRow);
+						// After folding the remaining options, ignoring nops, are:
+						//		1. Add Interactive Sync Tag
+						//		2. Remove Interactive Sync Tag
+						//		3. Add Sync Tag
+						//		4. Add Interactive Sync Tag + Add Sync Tag
+						// The versions using tag are folded into the addition in the above check, so 3 and 4 do not need to be folded here.
+						// Adding an interactive tag will be needed so leave that untouched, but do nop the remove interactive sync tag as
+						// there's guaranteed to be no interactive sync tag added at this point.
+						if (Optimizer.GetRight().IsType<FNopCommand>())
+						{
+							// Clearing out commands could leave right pointing at a nop, so skip that one and any following nops.
+							Optimizer.MoveToNextRight();
+						}
+						if (Optimizer.GetRight().IsType<FRemoveInteractiveSyncFromWorldTag>() ||
+							Optimizer.GetRight().IsType<FAddSyncFromWorldTag>())
+						{
+							Optimizer.ReplaceRight<FNopCommand>();
+						}
+					}
+					else
+					{
+						// This is a command that might require the order to remain stable so stop optimizations at this point.
+						break;
+					}
+				}
+				else
+				{
+					break;
+				}
+			}
+		}
+		Optimizer.ResetRightNextToLeft();
+		return Count;
+	}
+
+	void FCommandOptimizer::FoldSyncFromWorldTags(CompatibilityCommandBuffer::FOptimizer& Cluster)
+	{
+		// The following combinations are possible for interactive tags:
+		//		1. Add + Remove -> nop + nop
+		//		2. Remove + Add -> nop + Add
+		//		3. Add + Add -> Add + nop
+		//		4. Remove + Remove -> Remove + nop
+		// Sync from world tags are always after the interactive tags, so combinations are:
+		//		5. Add Interactive + Sync Tag -> Add Interactive + Sync Tag
+		//		6. Remove Interactive + Sync Tag -> Remove Interactive + nop (Remove Interactive also adds a sync tag).
+		//		7. Sync Tag + Sync Tag -> Sync Tag + nop
+		while (Cluster.IsValid())
+		{
+			const CompatibilityCommandBuffer::TCommandVariant& Left = Cluster.GetLeft();
+			const CompatibilityCommandBuffer::TCommandVariant& Right = Cluster.GetRight();
+
+			if (Left.IsType<FAddInteractiveSyncFromWorldTag>() && Right.IsType<FRemoveInteractiveSyncFromWorldTag>()) // 1
+			{
+				Cluster.ReplaceLeft<FNopCommand>();
+				Cluster.ReplaceRight<FNopCommand>();
+				
+				// Move forward to the next sync tag command.
+				Cluster.MoveLeftToRight();
+				Cluster.MoveToNextLeft();
+				for (; Cluster.IsValid(); Cluster.MoveToNextRight())
+				{
+					const CompatibilityCommandBuffer::TCommandVariant& NewLeft = Cluster.GetLeft();
+					if (NewLeft.IsType<FAddInteractiveSyncFromWorldTag>() ||
+						NewLeft.IsType<FRemoveInteractiveSyncFromWorldTag>() ||
+						NewLeft.IsType<FAddSyncFromWorldTag>())
+					{
+						break;
+					}
+				}
+			}
+			else if (Left.IsType<FRemoveInteractiveSyncFromWorldTag>() && Right.IsType<FAddInteractiveSyncFromWorldTag>()) // 2
+			{
+				Cluster.ReplaceLeft<FNopCommand>();
+				Cluster.MoveLeftToRight();
+			}
+			else if (
+				(Left.IsType<FAddInteractiveSyncFromWorldTag>() && Right.IsType<FAddInteractiveSyncFromWorldTag>()) || // 3
+				(Left.IsType<FRemoveInteractiveSyncFromWorldTag>() && Right.IsType<FRemoveInteractiveSyncFromWorldTag>()) || // 4
+				(Left.IsType<FRemoveInteractiveSyncFromWorldTag>() && Right.IsType<FAddSyncFromWorldTag>()) || // 6
+				(Left.IsType<FAddSyncFromWorldTag>() && Right.IsType<FAddSyncFromWorldTag>())) // 7
+			{
+				Cluster.ReplaceRight<FNopCommand>();
+				Cluster.MoveToNextRight();
+			}
+			else if (Left.IsType<FAddInteractiveSyncFromWorldTag>() && Right.IsType<FAddSyncFromWorldTag>()) // 5
+			{
+				Cluster.MoveLeftToRight();
+			}
+			else
+			{
+				Cluster.MoveToNextRight();
+			}
+		}
+	}
+
+	void FCommandOptimizer::operator()(const FAddInteractiveSyncFromWorldTag& Command)
+	{
+		CompatibilityCommandBuffer::FOptimizer Range = CreateRangeOptimizer(Command.Row);
+		FoldSyncFromWorldTags(Range);
+	}
+
+	void FCommandOptimizer::operator()(const FRemoveInteractiveSyncFromWorldTag& Command)
+	{
+		CompatibilityCommandBuffer::FOptimizer Range = CreateRangeOptimizer(Command.Row);
+		FoldSyncFromWorldTags(Range);
+	}
+
+	void FCommandOptimizer::operator()(const FAddSyncFromWorldTag& Command)
+	{
+		CompatibilityCommandBuffer::FOptimizer Range = CreateRangeOptimizer(Command.Row);
+		FoldSyncFromWorldTags(Range);
+	}
+
+	CompatibilityCommandBuffer::FOptimizer FCommandOptimizer::CreateRangeOptimizer(TypedElementDataStorage::RowHandle RowCluster)
+	{
+		return Optimizer.BranchOnLeft(
+			[RowCluster](const CompatibilityCommandBuffer::TCommandVariant& Command)
+			{
+				return FGetSourceRowHandle::Get(Command) == RowCluster;
+			});
+	}
+
+	void FCommandOptimizer::RunRightOnRowCluster(TypedElementDataStorage::RowHandle RowCluster)
+	{
+		CompatibilityCommandBuffer::FOptimizer SubOptimizer = Optimizer.BranchOnRight(
+			[RowCluster](const CompatibilityCommandBuffer::TCommandVariant& Command)
+			{
+				return FGetSourceRowHandle::Get(Command) == RowCluster;
+			});
+		FCommandOptimizer OptimizationSelector(SubOptimizer, ScratchBuffer);
+		if (SubOptimizer.IsValid())
+		{
+			Visit(OptimizationSelector, SubOptimizer.GetLeft());
+		}
+	}
+
+	void FCommandOptimizer::Run(CompatibilityCommandBuffer::FCollection& Commands, FTypedElementDatabaseScratchBuffer& ScratchBuffer)
+	{
+		CompatibilityCommandBuffer::FOptimizer Optimizer(Commands);
+		FCommandOptimizer OptimizationSelector(Optimizer, ScratchBuffer);
+		while (Optimizer.IsValid())
+		{
+			Visit(OptimizationSelector, Optimizer.GetLeft());
+			Optimizer.MoveToNextLeft();
+			Optimizer.ResetRightNextToLeft();
+		}
+	}
+} // namespace UE::EditorDataStorage
diff --git a/Engine/Plugins/Experimental/TypedElementsDataStorage/Source/TypedElementsDataStorage/Private/Commands/EditorDataStorageCompatibilityCommands.h b/Engine/Plugins/Experimental/TypedElementsDataStorage/Source/TypedElementsDataStorage/Private/Commands/EditorDataStorageCompatibilityCommands.h
new file mode 100644
--- /dev/null
+++ b/Engine/Plugins/Experimental/TypedElementsDataStorage/Source/TypedElementsDataStorage/Private/Commands/EditorDataStorageCompatibilityCommands.h
@@ -0,0 +1,387 @@
+// Copyright Epic Games, Inc. All Rights Reserved.
+
+#pragma once
+
+#include "Commands/EditorDataStorageCommandBuffer.h"
+#include "Delegates/IDelegateInstance.h"
+#include "Elements/Common/TypedElementHandles.h"
+#include "Elements/Interfaces/TypedElementDataStorageCompatibilityInterface.h"
+#include "UObject/WeakObjectPtr.h"
+#include "UObject/WeakObjectPtrTemplates.h"
+
+class ITypedElementDataStorageInterface;
+class FTypedElementDatabaseScratchBuffer;
+class UTypedElementMementoSystem;
+class UScriptStruct;
+class UTypedElementDatabaseCompatibility;
+class UObject;
+
+namespace UE::EditorDataStorage
+{
+	enum class EObjectType : uint8
+	{
+		Struct,
+		Class
+	};
+
+	/**
+	 * Objects with type info defined in either UScriptStruct or UClass can be stored into TEDS via the
+	 * ITypedElementDataStorageCompatibilityInterface
+	 * This is a discriminated union which aids with callbacks made when objects are added
+	 */
+	struct FObjectTypeInfo
+	{
+		EObjectType TypeInfoType;
+
+		union
+		{
+			const UScriptStruct* ScriptStruct;
+			const UClass* Class;
+		};
+
+		FObjectTypeInfo(const UScriptStruct* InScriptStruct);
+		FObjectTypeInfo(const UClass* InClass);
+
+		FName GetFName() const;
+	};
+
+	struct FNopCommand;
+
+	struct FTypeInfoReinstanced final
+	{
+		TWeakObjectPtr<UStruct> Original;
+		TWeakObjectPtr<UStruct> Reinstanced;
+	};
+
+	struct FTypeBatchInfoReinstanced final
+	{
+		TArrayView<FTypeInfoReinstanced> Batch;
+
+		template<typename T>
+		static FTypeInfoReinstanced* FindObject(TArrayView<FTypeInfoReinstanced> Range, const TWeakObjectPtr<T>& Object);
+		template<typename T>
+		static TWeakObjectPtr<T> FindObjectRecursively(TArrayView<FTypeInfoReinstanced> Range, const TWeakObjectPtr<T>& Object);
+	};
+
+	struct FRegisterTypeTableAssociation final
+	{
+		TWeakObjectPtr<UStruct> TypeInfo;
+		TypedElementDataStorage::TableHandle Table;
+	};
+
+	struct FAddCompatibleUObject final
+	{
+		TWeakObjectPtr<UObject> Object;
+		TypedElementDataStorage::RowHandle Row;
+		TypedElementDataStorage::TableHandle Table;
+	};
+
+	using ObjectAddedCallback = TFunction<void(
+		const void* /*Object*/, 
+		const FObjectTypeInfo&, /*Type information*/
+		TypedElementDataStorage::RowHandle /*Row*/)>;
+	struct FRegisterObjectAddedCallback final
+	{
+		ObjectAddedCallback Callback;
+		FDelegateHandle Handle;
+	};
+
+	struct FUnregisterObjectAddedCallback final
+	{
+		FDelegateHandle Handle;
+	};
+
+	using ObjectRemovedCallback = TFunction<void(
+		const void* /*Object*/, 
+		const FObjectTypeInfo&, /*Type information*/
+		TypedElementDataStorage::RowHandle /*Row*/)>;
+
+	struct FBatchAddCompatibleUObject final
+	{
+		TWeakObjectPtr<UObject>* ObjectArray;
+		TypedElementDataStorage::RowHandle* RowArray;
+		TypedElementDataStorage::TableHandle Table;
+		uint32 Count;
+	};
+
+	struct FAddCompatibleExternalObject final
+	{
+		void* Object;
+		TWeakObjectPtr<UScriptStruct> TypeInfo;
+		TypedElementDataStorage::RowHandle Row;
+		TypedElementDataStorage::TableHandle Table;
+	};
+
+	struct FCreateMemento final
+	{
+		TypedElementDataStorage::RowHandle ReservedMementoRow;
+		TypedElementDataStorage::RowHandle TargetRow;
+	};
+
+	struct FRestoreMemento final
+	{
+		TypedElementDataStorage::RowHandle MementoRow;
+		TypedElementDataStorage::RowHandle TargetRow;
+	};
+
+	struct FDestroyMemento final
+	{
+		TypedElementDataStorage::RowHandle MementoRow;
+	};
+
+	struct FRemoveCompatibleUObject final
+	{
+		const UObject* Object;
+		TypedElementDataStorage::RowHandle ObjectRow = TypedElementDataStorage::InvalidRowHandle;
+	};
+
+	struct FRemoveCompatibleExternalObject final
+	{
+		void* Object;
+		TypedElementDataStorage::RowHandle ObjectRow;
+	};
+
+	struct FBatchAddCompatibleExternalObject final
+	{
+		void** ObjectArray;
+		TWeakObjectPtr<UScriptStruct>* TypeInfoArray;
+		TypedElementDataStorage::RowHandle* RowArray;
+		TypedElementDataStorage::TableHandle Table;
+		uint32 Count;
+	};
+
+	struct FAddSyncFromWorldTag final
+	{
+		static UScriptStruct* Type;
+
+		TObjectKey<const UObject> Target;
+		TypedElementDataStorage::RowHandle Row;
+	};
+
+	struct FAddInteractiveSyncFromWorldTag final
+	{
+		static UScriptStruct* Type;
+
+		TObjectKey<const UObject> Target;
+		TypedElementDataStorage::RowHandle Row;
+	};
+
+	struct FRemoveInteractiveSyncFromWorldTag final
+	{
+		TObjectKey<const UObject> Target;
+		TypedElementDataStorage::RowHandle Row;
+	};
+
+	using CompatibilityCommandBuffer = FCommandBuffer
+	<
+		FTypeInfoReinstanced,
+		FTypeBatchInfoReinstanced,
+		FRegisterTypeTableAssociation,
+		FRegisterObjectAddedCallback,
+		FUnregisterObjectAddedCallback,
+		
+		FAddCompatibleUObject,
+		FBatchAddCompatibleUObject,
+		FAddCompatibleExternalObject,
+		FBatchAddCompatibleExternalObject,
+		
+		FCreateMemento,
+		FRestoreMemento,
+		FDestroyMemento,
+		
+		FRemoveCompatibleUObject,
+		FRemoveCompatibleExternalObject,
+		
+		FAddInteractiveSyncFromWorldTag,
+		FRemoveInteractiveSyncFromWorldTag,
+		FAddSyncFromWorldTag
+	>;
+
+	/**
+	 * Patches data in preparation for processing. This can include fixing tables like the type-to-table map if type information has 
+	 * changed.
+	 */
+	struct FPatchData final
+	{
+		struct FPatchCommand final
+		{
+			explicit FPatchCommand(TArrayView<FTypeInfoReinstanced>& InReinstances);
+
+			template<typename T>
+			void operator()(T&) {}
+
+			void operator()(FRegisterTypeTableAssociation& Command);
+			void operator()(FAddCompatibleExternalObject& Command);
+			void operator()(FBatchAddCompatibleExternalObject& Command);
+
+			TArrayView<FTypeInfoReinstanced>& Reinstances;
+		};
+
+		static bool IsPatchingRequired(const CompatibilityCommandBuffer::FCollection& Commands);
+		static void RunPatch(CompatibilityCommandBuffer::FCollection& Commands, UTypedElementDatabaseCompatibility& StorageCompat,
+			FTypedElementDatabaseScratchBuffer& ScratchBuffer);
+	};
+
+	/** Prepares each command for further processing, e.g. resolving the target table. */
+	struct FPrepareCommands final
+	{
+		FPrepareCommands(ITypedElementDataStorageInterface& InStorage, UTypedElementDatabaseCompatibility& InStorageCompat,
+			CompatibilityCommandBuffer::FCollection& InCommands);
+
+		template<typename T>
+		void operator()(T&) {}
+		void operator()(FAddCompatibleUObject& Object);
+		void operator()(FAddCompatibleExternalObject& Object);
+		void operator()(FRemoveCompatibleUObject& Command);
+		void operator()(FRemoveCompatibleExternalObject& Command);
+		void operator()(FAddInteractiveSyncFromWorldTag& Command);
+		void operator()(FRemoveInteractiveSyncFromWorldTag& Command);
+		void operator()(FAddSyncFromWorldTag& Command);
+
+		ITypedElementDataStorageInterface& Storage;
+		UTypedElementDatabaseCompatibility& StorageCompat;
+		CompatibilityCommandBuffer::FCollection& Commands;
+		int32 CurrentIndex = 0;
+
+		static void RunPreparation(ITypedElementDataStorageInterface& Storage, UTypedElementDatabaseCompatibility& StorageCompat, 
+			CompatibilityCommandBuffer::FCollection& Commands);
+	};
+
+	/** Retrieve the source row handle from a command. */
+	struct FGetSourceRowHandle
+	{
+		template<typename T>
+		TypedElementDataStorage::RowHandle operator()(const T&){ return TypedElementDataStorage::InvalidRowHandle; }
+		TypedElementDataStorage::RowHandle operator()(const FAddCompatibleUObject& Command);
+		TypedElementDataStorage::RowHandle operator()(const FAddCompatibleExternalObject& Command);
+		TypedElementDataStorage::RowHandle operator()(const FCreateMemento& Command);
+		TypedElementDataStorage::RowHandle operator()(const FRestoreMemento& Command);
+		TypedElementDataStorage::RowHandle operator()(const FDestroyMemento& Command);
+		TypedElementDataStorage::RowHandle operator()(const FRemoveCompatibleUObject& Command);
+		TypedElementDataStorage::RowHandle operator()(const FRemoveCompatibleExternalObject& Command);
+		TypedElementDataStorage::RowHandle operator()(const FAddInteractiveSyncFromWorldTag& Command);
+		TypedElementDataStorage::RowHandle operator()(const FRemoveInteractiveSyncFromWorldTag& Command);
+		TypedElementDataStorage::RowHandle operator()(const FAddSyncFromWorldTag& Command);
+		static TypedElementDataStorage::RowHandle Get(const CompatibilityCommandBuffer::TCommandVariant& Command);
+	};
+
+	/** 
+	 * Group id for a command. Commands within the same group are not reordered for commands on the same row. The default order
+	 * follows the order in which the commands are declared in the command buffer.
+	 */
+	struct FGetCommandGroupId
+	{
+		template<typename T>
+		SIZE_T operator()(const T&) { return CompatibilityCommandBuffer::TCommandVariant::template IndexOfType<T>(); }
+		
+		// FAddCompatibleUObject group
+		SIZE_T operator()(const FBatchAddCompatibleUObject& Command);
+		SIZE_T operator()(const FAddCompatibleExternalObject& Command);
+		SIZE_T operator()(const FBatchAddCompatibleExternalObject& Command);
+		SIZE_T operator()(const FCreateMemento& Command);
+		SIZE_T operator()(const FRestoreMemento& Command);
+		SIZE_T operator()(const FDestroyMemento& Command);
+		SIZE_T operator()(const FRemoveCompatibleUObject& Command);
+		SIZE_T operator()(const FRemoveCompatibleExternalObject& Command);
+		// FAddInteractiveSyncFromWorldTag group
+		SIZE_T operator()(const FRemoveInteractiveSyncFromWorldTag& Command);
+		static SIZE_T Get(const CompatibilityCommandBuffer::TCommandVariant& Command);
+	};
+
+	struct FSorter final
+	{
+		static void SortCommands(CompatibilityCommandBuffer::FCollection& Commands);
+	};
+
+	/** Executes the commands in the command buffer for TEDS Compatibility. */
+	struct FCommandProcessor final
+	{
+		FCommandProcessor(ITypedElementDataStorageInterface& InStorage, UTypedElementDatabaseCompatibility& InStorageCompatibility);
+
+		void SetupRow(TypedElementDataStorage::RowHandle Row, UObject* Object);
+		void SetupRow(TypedElementDataStorage::RowHandle Row, void* Object, TWeakObjectPtr<UScriptStruct> TypeInfo);
+		
+		void operator()(FNopCommand&) {}
+		void operator()(FTypeInfoReinstanced&) {}
+		void operator()(FTypeBatchInfoReinstanced& Command);
+		void operator()(FRegisterTypeTableAssociation& Command);
+		void operator()(FRegisterObjectAddedCallback& Command);
+		void operator()(FUnregisterObjectAddedCallback& Command);
+		void operator()(FAddCompatibleUObject& Object);
+		void operator()(FBatchAddCompatibleUObject& Batch);
+		void operator()(FAddCompatibleExternalObject& Object);
+		void operator()(FBatchAddCompatibleExternalObject& Object);
+		void operator()(FCreateMemento& Command);
+		void operator()(FRestoreMemento& Command);
+		void operator()(FDestroyMemento& Command);
+		void operator()(FRemoveCompatibleUObject& Command);
+		void operator()(FRemoveCompatibleExternalObject& Command);
+		void operator()(FAddInteractiveSyncFromWorldTag& Command);
+		void operator()(FRemoveInteractiveSyncFromWorldTag& Command);
+		void operator()(FAddSyncFromWorldTag& Command);
+
+		ITypedElementDataStorageInterface& Storage;
+		UTypedElementDatabaseCompatibility& StorageCompatibility;
+		UTypedElementMementoSystem& MementoSystem;
+	};
+
+	struct FRecordCommands final
+	{
+		void operator()(const FNopCommand& Command);
+		void operator()(const FTypeInfoReinstanced& Command);
+		void operator()(const FTypeBatchInfoReinstanced& Command);
+		void operator()(const FRegisterTypeTableAssociation& Command);
+		void operator()(const FRegisterObjectAddedCallback& Command);
+		void operator()(const FUnregisterObjectAddedCallback& Command);
+		void operator()(const FAddCompatibleUObject& Command);
+		void operator()(const FBatchAddCompatibleUObject& Command);
+		void operator()(const FAddCompatibleExternalObject& Command);
+		void operator()(const FBatchAddCompatibleExternalObject& Command);
+		void operator()(const FCreateMemento& Command);
+		void operator()(const FRestoreMemento& Command);
+		void operator()(const FDestroyMemento& Command);
+		void operator()(const FRemoveCompatibleUObject& Command);
+		void operator()(const FRemoveCompatibleExternalObject& Command);
+		void operator()(const FAddInteractiveSyncFromWorldTag& Command);
+		void operator()(const FRemoveInteractiveSyncFromWorldTag& Command);
+		void operator()(const FAddSyncFromWorldTag& Command);
+
+		FString CommandDescriptions;
+		bool bIncludeNops = false;
+
+		static FString PrintToString(CompatibilityCommandBuffer::FCollection& Commands, bool bIncludeNops);
+	};
+
+	/** Looks at the current command and applies the next optimization. */
+	struct FCommandOptimizer final
+	{
+		// TODO: Allow [add object + restore memento] to be folded into one and [create memento + remove object].
+		//		Not sure how that will fit into the batch add though.
+		FCommandOptimizer(CompatibilityCommandBuffer::FOptimizer& InOptimizer, FTypedElementDatabaseScratchBuffer& InScratchBuffer);
+
+		template<typename T>
+		void operator()(const T&) {}
+		void operator()(const FAddCompatibleUObject& Command);
+		void operator()(const FAddCompatibleExternalObject& Command);
+		void operator()(const FRemoveCompatibleUObject& Command);
+		void operator()(const FRemoveCompatibleExternalObject& Command);
+		void operator()(const FAddInteractiveSyncFromWorldTag& Command);
+		void operator()(const FRemoveInteractiveSyncFromWorldTag& Command);
+		void operator()(const FAddSyncFromWorldTag& Command);
+
+		CompatibilityCommandBuffer::FOptimizer& Optimizer;
+		FTypedElementDatabaseScratchBuffer& ScratchBuffer;
+
+		template<typename AddCommandType>
+		int32 FoldCommandsForAdd(const AddCommandType& Command);
+		void FoldSyncFromWorldTags(CompatibilityCommandBuffer::FOptimizer& Cluster);
+		/**
+		 * Creates a new optimizer starting at the left of the current optimizer that will run only as long as commands on the right use 
+		 * the same row as the provided argument.
+		 */
+		CompatibilityCommandBuffer::FOptimizer CreateRangeOptimizer(TypedElementDataStorage::RowHandle RowCluster);
+		/** Keeps processing a subset of commands until the right no longer has a command with the provided source row. */
+		void RunRightOnRowCluster(TypedElementDataStorage::RowHandle RowCluster);
+		static void Run(CompatibilityCommandBuffer::FCollection& Commands, FTypedElementDatabaseScratchBuffer& ScratchBuffer);
+	};
+} // namespace UE::EditorDataStorage
diff --git a/Engine/Plugins/Experimental/TypedElementsDataStorage/Source/TypedElementsDataStorage/Private/Compatibility/TypedElementObjectReinstancingManager.cpp b/Engine/Plugins/Experimental/TypedElementsDataStorage/Source/TypedElementsDataStorage/Private/Compatibility/TypedElementObjectReinstancingManager.cpp
--- a/Engine/Plugins/Experimental/TypedElementsDataStorage/Source/TypedElementsDataStorage/Private/Compatibility/TypedElementObjectReinstancingManager.cpp
+++ b/Engine/Plugins/Experimental/TypedElementsDataStorage/Source/TypedElementsDataStorage/Private/Compatibility/TypedElementObjectReinstancingManager.cpp
@@ -18,6 +18,9 @@ UTypedElementObjectReinstancingManager::UTypedElementObjectReinstancingManager()
 
 void UTypedElementObjectReinstancingManager::Initialize(UTypedElementDatabase& InDatabase, UTypedElementDatabaseCompatibility& InDataStorageCompatibility)
 {
+	using namespace TypedElementDataStorage;
+	using namespace UE::EditorDataStorage;
+
 	Database = &InDatabase;
 	DataStorageCompatibility = &InDataStorageCompatibility;
 	
@@ -26,7 +29,7 @@ void UTypedElementObjectReinstancingManager::Initialize(UTypedElementDatabase& I
 	ReinstancingCallbackHandle = 
 		FCoreUObjectDelegates::OnObjectsReinstanced.AddUObject(this, &UTypedElementObjectReinstancingManager::HandleOnObjectsReinstanced);
 	ObjectRemovedCallbackHandle = DataStorageCompatibility->RegisterObjectRemovedCallback(
-		[this](const void* Object, const FTypedElementDatabaseCompatibilityObjectTypeInfo& TypeInfo, TypedElementRowHandle Row)
+		[this](const void* Object, const FObjectTypeInfo& TypeInfo, RowHandle Row)
 		{
 			HandleOnObjectPreRemoved(Object, TypeInfo, Row);
 		});
@@ -55,10 +58,12 @@ void UTypedElementObjectReinstancingManager::UpdateCompleted()
 }
 
 void UTypedElementObjectReinstancingManager::HandleOnObjectPreRemoved(
-	const void* Object, const FTypedElementDatabaseCompatibilityObjectTypeInfo& TypeInfo, TypedElementRowHandle ObjectRow)
+	const void* Object, 
+	const UE::EditorDataStorage::FObjectTypeInfo& TypeInfo, 
+	TypedElementDataStorage::RowHandle ObjectRow)
 {
 	// This is the chance to record the old object to memento
-	TypedElementRowHandle Memento = Database->GetEnvironment()->GetMementoSystem().CreateMemento(ObjectRow);
+	TypedElementDataStorage::RowHandle Memento = Database->GetEnvironment()->GetMementoSystem().CreateMemento(ObjectRow);
 	OldObjectToMementoMap.Add(Object, Memento);
 }
 
diff --git a/Engine/Plugins/Experimental/TypedElementsDataStorage/Source/TypedElementsDataStorage/Private/Compatibility/TypedElementObjectReinstancingManager.h b/Engine/Plugins/Experimental/TypedElementsDataStorage/Source/TypedElementsDataStorage/Private/Compatibility/TypedElementObjectReinstancingManager.h
--- a/Engine/Plugins/Experimental/TypedElementsDataStorage/Source/TypedElementsDataStorage/Private/Compatibility/TypedElementObjectReinstancingManager.h
+++ b/Engine/Plugins/Experimental/TypedElementsDataStorage/Source/TypedElementsDataStorage/Private/Compatibility/TypedElementObjectReinstancingManager.h
@@ -14,6 +14,11 @@ class ITypedElementDataStorageInterface;
 class UTypedElementDatabaseCompatibility;
 class UTypedElementMementoSystem;
 
+namespace UE::EditorDataStorage
+{
+	struct FObjectTypeInfo;
+}
+
 UCLASS(Transient)
 class UTypedElementObjectReinstancingManager : public UObject
 {
@@ -27,7 +32,9 @@ public:
 private:
 	void UpdateCompleted();
 	void HandleOnObjectPreRemoved(
-		const void* Object, const FTypedElementDatabaseCompatibilityObjectTypeInfo& TypeInfo, TypedElementDataStorage::RowHandle ObjectRow);
+		const void* Object, 
+		const UE::EditorDataStorage::FObjectTypeInfo& TypeInfo, 
+		TypedElementDataStorage::RowHandle ObjectRow);
 	void HandleOnObjectsReinstanced(const FCoreUObjectDelegates::FReplacementObjectMap& ObjectReplacementMap);
 
 	UPROPERTY()
diff --git a/Engine/Plugins/Experimental/TypedElementsDataStorage/Source/TypedElementsDataStorage/Private/Memento/TypedElementMementoSystem.cpp b/Engine/Plugins/Experimental/TypedElementsDataStorage/Source/TypedElementsDataStorage/Private/Memento/TypedElementMementoSystem.cpp
--- a/Engine/Plugins/Experimental/TypedElementsDataStorage/Source/TypedElementsDataStorage/Private/Memento/TypedElementMementoSystem.cpp
+++ b/Engine/Plugins/Experimental/TypedElementsDataStorage/Source/TypedElementsDataStorage/Private/Memento/TypedElementMementoSystem.cpp
@@ -4,6 +4,7 @@
 
 #include "Elements/Common/TypedElementDataStorageLog.h"
 #include "Elements/Framework/TypedElementQueryBuilder.h"
+#include "GlobalLock.h"
 #include "Memento/TypedElementMementoTranslators.h"
 #include "TypedElementMementoRowTypes.h"
 #include "TypedElementDatabase.h"
@@ -11,6 +12,10 @@
 UTypedElementMementoSystem::UTypedElementMementoSystem(ITypedElementDataStorageInterface& InDataStorage)
 	: DataStorage(InDataStorage)
 {
+	using namespace UE::EditorDataStorage;
+
+	FScopedExclusiveLock Lock(EGlobalLockScope::Public);
+
 	// Register tables that will be used by reinstancing
 	MementoRowBaseTable = DataStorage.RegisterTable<FTypedElementMementoTag>(TEXT("MementoRowBaseTable"));
 
@@ -36,8 +41,29 @@ UTypedElementMementoSystem::UTypedElementMementoSystem(ITypedElementDataStorageI
 TypedElementRowHandle UTypedElementMementoSystem::CreateMemento(TypedElementDataStorage::RowHandle SourceRow)
 {
 	using namespace TypedElementDataStorage;
+	using namespace UE::EditorDataStorage;
+
+	FScopedSharedLock Lock(EGlobalLockScope::Public);
 
 	RowHandle MementoRow = DataStorage.AddRow(MementoRowBaseTable);
+	CreateMementoInternal(MementoRow, SourceRow);
+	return MementoRow;
+}
+
+void UTypedElementMementoSystem::CreateMemento(TypedElementDataStorage::RowHandle ReservedMementoRow, TypedElementDataStorage::RowHandle SourceRow)
+{
+	using namespace UE::EditorDataStorage;
+
+	FScopedSharedLock Lock(EGlobalLockScope::Public);
+	
+	DataStorage.AddRow(ReservedMementoRow, MementoRowBaseTable);
+	CreateMementoInternal(ReservedMementoRow, SourceRow);
+}
+
+void UTypedElementMementoSystem::CreateMementoInternal(TypedElementDataStorage::RowHandle MementoRow, TypedElementDataStorage::RowHandle SourceRow)
+{
+	using namespace TypedElementDataStorage;
+	using namespace UE::EditorDataStorage;
 	
 	for (const UTypedElementMementoTranslatorBase* Translator : MementoTranslators)
 	{
@@ -48,6 +74,7 @@ TypedElementRowHandle UTypedElementMementoSystem::CreateMemento(TypedElementData
 				[Translator, SourceColumn](void* MementoColumn, const UScriptStruct& ColumnType)
 				{
 					ColumnType.InitializeStruct(MementoColumn);
+					FScopedSharedLock Lock(EGlobalLockScope::Public);
 					Translator->TranslateColumnToMemento(SourceColumn, MementoColumn);
 				},
 				[](const UScriptStruct& ColumnType, void* Destination, void* Source)
@@ -55,16 +82,18 @@ TypedElementRowHandle UTypedElementMementoSystem::CreateMemento(TypedElementData
 					ColumnType.CopyScriptStruct(Destination, Source);
 				});
 
-			UE_LOG(LogTypedElementDataStorage, VeryVerbose, 
+			UE_LOG(LogTypedElementDataStorage, VeryVerbose,
 				TEXT("Column->Memento: %llu -> %llu"), SourceRow, MementoRow);
 		}
 	}
-
-	return MementoRow;
 }
 
 void UTypedElementMementoSystem::RestoreMemento(TypedElementDataStorage::RowHandle MementoRow, TypedElementDataStorage::RowHandle TargetRow)
 {
+	using namespace UE::EditorDataStorage;
+
+	FScopedSharedLock Lock(EGlobalLockScope::Public);
+
 	for (const UTypedElementMementoTranslatorBase* Translator : MementoTranslators)
 	{
 		if (void* MementoColumn = DataStorage.GetColumnData(MementoRow, Translator->GetMementoType()))
@@ -74,6 +103,7 @@ void UTypedElementMementoSystem::RestoreMemento(TypedElementDataStorage::RowHand
 				[Translator, MementoColumn](void* TargetColumn, const UScriptStruct& ColumnType)
 				{
 					ColumnType.InitializeStruct(TargetColumn);
+					FScopedSharedLock Lock(EGlobalLockScope::Public);
 					Translator->TranslateMementoToColumn(MementoColumn, TargetColumn);
 				},
 				[](const UScriptStruct& ColumnType, void* Destination, void* Source)
@@ -89,6 +119,8 @@ void UTypedElementMementoSystem::RestoreMemento(TypedElementDataStorage::RowHand
 
 void UTypedElementMementoSystem::DestroyMemento(TypedElementDataStorage::RowHandle MementoRow)
 {
+	// No need to lock this as no internal data is used.
+
 	checkf(DataStorage.IsRowAvailable(MementoRow) && DataStorage.HasColumns<FTypedElementMementoTag>(MementoRow),
 		TEXT("Deleting memento row that's not marked as such."));
 	DataStorage.RemoveRow(MementoRow);
diff --git a/Engine/Plugins/Experimental/TypedElementsDataStorage/Source/TypedElementsDataStorage/Private/Memento/TypedElementMementoSystem.h b/Engine/Plugins/Experimental/TypedElementsDataStorage/Source/TypedElementsDataStorage/Private/Memento/TypedElementMementoSystem.h
--- a/Engine/Plugins/Experimental/TypedElementsDataStorage/Source/TypedElementsDataStorage/Private/Memento/TypedElementMementoSystem.h
+++ b/Engine/Plugins/Experimental/TypedElementsDataStorage/Source/TypedElementsDataStorage/Private/Memento/TypedElementMementoSystem.h
@@ -15,10 +15,13 @@ public:
 	explicit UTypedElementMementoSystem(ITypedElementDataStorageInterface& InDataStorage);
 	
 	TypedElementDataStorage::RowHandle CreateMemento(TypedElementDataStorage::RowHandle SourceRow);
+	void CreateMemento(TypedElementDataStorage::RowHandle ReservedMementoRow, TypedElementDataStorage::RowHandle SourceRow);
 	void RestoreMemento(TypedElementDataStorage::RowHandle MementoRow,  TypedElementDataStorage::RowHandle TargetRow);
 	void DestroyMemento(TypedElementDataStorage::RowHandle MementoRow);
 
 private:
+	void CreateMementoInternal(TypedElementDataStorage::RowHandle MementoRow, TypedElementDataStorage::RowHandle SourceRow);
+
 	TArray<const UTypedElementMementoTranslatorBase*> MementoTranslators;
 	TypedElementDataStorage::TableHandle MementoRowBaseTable;
 	ITypedElementDataStorageInterface& DataStorage;
diff --git a/Engine/Plugins/Experimental/TypedElementsDataStorage/Source/TypedElementsDataStorage/Private/TypedElementDatabaseCompatibility.cpp b/Engine/Plugins/Experimental/TypedElementsDataStorage/Source/TypedElementsDataStorage/Private/TypedElementDatabaseCompatibility.cpp
--- a/Engine/Plugins/Experimental/TypedElementsDataStorage/Source/TypedElementsDataStorage/Private/TypedElementDatabaseCompatibility.cpp
+++ b/Engine/Plugins/Experimental/TypedElementsDataStorage/Source/TypedElementsDataStorage/Private/TypedElementDatabaseCompatibility.cpp
@@ -16,30 +16,67 @@
 #include "Elements/Framework/TypedElementIndexHasher.h"
 #include "Elements/Framework/TypedElementQueryBuilder.h"
 #include "HAL/IConsoleManager.h"
+#include "Logging/LogMacros.h"
 #include "Memento/TypedElementMementoRowTypes.h"
 #include "TypedElementDatabase.h"
 #include "TypedElementDatabaseEnvironment.h"
 #include "TypedElementDataStorageProfilingMacros.h"
 
+DEFINE_LOG_CATEGORY_STATIC(LogTedsCompat, Log, All);
+
 namespace TypedElementDataStorage
 {
 	bool bIntegrateWithGC = true;
 	FAutoConsoleVariableRef CVarIntegrateWithGC(
 		TEXT("TEDS.Feature.IntegrateWithGC"),
 		bIntegrateWithGC,
-		TEXT("Enabled monitoring GC events from TEDS Compat so objects are cleared up at the last minute, before they're destroyed."));
+		TEXT("Enables actors being removed through the garbage collection instead of requiring explicit removal."));
+
+	bool bUseCommandBuffer = false;
+	FAutoConsoleVariableRef CVarUseCommandBufferInCompat(
+		TEXT("TEDS.Feature.UseCommandBufferInCompat"),
+		bUseCommandBuffer,
+		TEXT("Use the command buffer to defer TEDS Compatibility commands."));
+
+	bool bUseDeferredRemovesInCompat = false;
+	FAutoConsoleVariableRef CVarUseDeferredRemovesInCompat(
+		TEXT("TEDS.Feature.UseDeferredRemovesInCompat"),
+		bUseDeferredRemovesInCompat,
+		TEXT("If the command buffer in TEDS Compatibility is enabled, setting this to true will cause removes to be queued instead "
+			"of immediately executed."));
+
+	bool bOptimizeCommandBuffer = true;
+	FAutoConsoleVariableRef CVarOptimizeCommandBufferInCompat(
+		TEXT("TEDS.Debug.OptimizeCommandBufferInCompat"),
+		bOptimizeCommandBuffer,
+		TEXT("If true, the command buffer used in TEDS Compat is optimized, otherwise the optimization phase is skipped."));
+
+	int PrintCompatCommandBuffer = 0;
+	FAutoConsoleVariableRef CVarPrintCompatCommandBuffer(
+		TEXT("TEDS.Debug.PrintCompatCommandBuffer"),
+		PrintCompatCommandBuffer,
+		TEXT("If enabled and TEDS Compat uses the command buffer, then the list of pending commands is printed before being execute.\n"
+			"0 - disable\n"
+			"1 - summarize number of nops\n"
+			"2 - include nops"));
 }
 
 static const FName IntegrateWithGCName(TEXT("IntegrateWithGC"));
+static const FName CompatibilityUsesCommandBufferExtensionName(TEXT("CompatiblityUsesCommandBuffer"));
 
 void UTypedElementDatabaseCompatibility::Initialize(UTypedElementDatabase* InStorage)
 {
 	using namespace TypedElementDataStorage;
+	using namespace UE::EditorDataStorage;
 
-	checkf(InStorage, TEXT("Typed Element's Database compatibility manager is being initialized with an invalid storage target."));
+	checkf(InStorage, TEXT("TEDS Compatibility is being initialized with an invalid storage target."));
 	
+	FScopedExclusiveLock Lock(EGlobalLockScope::Public);
+
 	Storage = InStorage;
 	Environment = InStorage->GetEnvironment();
+	QueuedCommands.Initialize(Environment->GetScratchBuffer());
+	
 	Prepare();
 
 	InStorage->OnUpdate().AddUObject(this, &UTypedElementDatabaseCompatibility::Tick);
@@ -58,7 +95,10 @@ void UTypedElementDatabaseCompatibility::Initialize(UTypedElementDatabase* InSto
 void UTypedElementDatabaseCompatibility::Deinitialize()
 {
 	using namespace TypedElementDataStorage;
+	using namespace UE::EditorDataStorage;
 
+	FScopedExclusiveLock Lock(EGlobalLockScope::Public);
+	
 	for (TPair<UWorld*, FDelegateHandle>& It : ActorDestroyedDelegateHandles)
 	{
 		It.Key->RemoveOnActorDestroyededHandler(It.Value);
@@ -78,37 +118,79 @@ void UTypedElementDatabaseCompatibility::Deinitialize()
 
 void UTypedElementDatabaseCompatibility::RegisterRegistrationFilter(ObjectRegistrationFilter Filter)
 {
+	using namespace UE::EditorDataStorage;
+
+	FScopedExclusiveLock Lock(EGlobalLockScope::Public);
 	ObjectRegistrationFilters.Add(MoveTemp(Filter));
 }
 
 void UTypedElementDatabaseCompatibility::RegisterDealiaserCallback(ObjectToRowDealiaser Dealiaser)
 {
+	using namespace UE::EditorDataStorage;
+
+	FScopedExclusiveLock Lock(EGlobalLockScope::Public);
 	ObjectToRowDialiasers.Add(MoveTemp(Dealiaser));
 }
 
 void UTypedElementDatabaseCompatibility::RegisterTypeTableAssociation(
 	TObjectPtr<UStruct> TypeInfo, TypedElementDataStorage::TableHandle Table)
 {
-	TypeToTableMap.Add(TypeInfo, Table);
+	using namespace UE::EditorDataStorage;
+	
+	if (TypedElementDataStorage::bUseCommandBuffer)
+	{
+		QueuedCommands.AddCommand(UE::EditorDataStorage::FRegisterTypeTableAssociation{ .TypeInfo = TypeInfo, .Table = Table });
+	}
+	else
+	{
+		FScopedExclusiveLock Lock(EGlobalLockScope::Public);
+		TypeToTableMap.Add(TypeInfo, Table);
+	}
 }
 
-FDelegateHandle UTypedElementDatabaseCompatibility::RegisterObjectAddedCallback(ObjectAddedCallback&& OnObjectAdded)
+FDelegateHandle UTypedElementDatabaseCompatibility::RegisterObjectAddedCallback(UE::EditorDataStorage::ObjectAddedCallback&& OnObjectAdded)
 {
+	using namespace UE::EditorDataStorage;
+
 	FDelegateHandle Handle(FDelegateHandle::GenerateNewHandle);
-	ObjectAddedCallbackList.Emplace(MoveTemp(OnObjectAdded), Handle);
+	if (TypedElementDataStorage::bUseCommandBuffer)
+	{
+		QueuedCommands.AddCommand(UE::EditorDataStorage::FRegisterObjectAddedCallback{ .Callback = MoveTemp(OnObjectAdded), .Handle = Handle });
+	}
+	else
+	{
+		FScopedExclusiveLock Lock(EGlobalLockScope::Public);
+		ObjectAddedCallbackList.Emplace(MoveTemp(OnObjectAdded), Handle);
+	}
 	return Handle;
 }
 
 void UTypedElementDatabaseCompatibility::UnregisterObjectAddedCallback(FDelegateHandle Handle)
 {
-	ObjectAddedCallbackList.RemoveAll([Handle](const TPair<ObjectAddedCallback, FDelegateHandle>& Element)->bool
+	using namespace UE::EditorDataStorage;
+
+	if (TypedElementDataStorage::bUseCommandBuffer)
 	{
-		return Element.Value == Handle;
-	});
+		QueuedCommands.AddCommand(UE::EditorDataStorage::FUnregisterObjectAddedCallback{ .Handle = Handle });
+	}
+	else
+	{
+		FScopedExclusiveLock Lock(EGlobalLockScope::Public);
+		ObjectAddedCallbackList.RemoveAll(
+			[Handle](const TPair<UE::EditorDataStorage::ObjectAddedCallback, FDelegateHandle>& Element)->bool
+			{
+				return Element.Value == Handle;
+			});
+	}
 }
 
-FDelegateHandle UTypedElementDatabaseCompatibility::RegisterObjectRemovedCallback(ObjectRemovedCallback&& OnObjectAdded)
+FDelegateHandle UTypedElementDatabaseCompatibility::RegisterObjectRemovedCallback(UE::EditorDataStorage::ObjectRemovedCallback&& OnObjectAdded)
 {
+	using namespace UE::EditorDataStorage;
+
+	FScopedExclusiveLock Lock(EGlobalLockScope::Public);
+
+	// Since removing object has be immediately executed in some situation, adding the callback can not be delayed through the command buffer.
 	FDelegateHandle Handle(FDelegateHandle::GenerateNewHandle);
 	PreObjectRemovedCallbackList.Emplace(MoveTemp(OnObjectAdded), Handle);
 	return Handle;
@@ -116,7 +198,12 @@ FDelegateHandle UTypedElementDatabaseCompatibility::RegisterObjectRemovedCallbac
 
 void UTypedElementDatabaseCompatibility::UnregisterObjectRemovedCallback(FDelegateHandle Handle)
 {
-	PreObjectRemovedCallbackList.RemoveAll([Handle](const TPair<ObjectRemovedCallback, FDelegateHandle>& Element)->bool
+	using namespace UE::EditorDataStorage;
+
+	// Since removing object has be immediately executed in some situation, adding the callback can not be delayed through the command buffer.
+	FScopedExclusiveLock Lock(EGlobalLockScope::Public);
+
+	PreObjectRemovedCallbackList.RemoveAll([Handle](const TPair<UE::EditorDataStorage::ObjectRemovedCallback, FDelegateHandle>& Element)->bool
 	{
 		return Element.Value == Handle;
 	});
@@ -124,31 +211,37 @@ void UTypedElementDatabaseCompatibility::UnregisterObjectRemovedCallback(FDelega
 
 TypedElementRowHandle UTypedElementDatabaseCompatibility::AddCompatibleObjectExplicit(UObject* Object)
 {
+	// Because AddCompatibleObjectExplicitTransactionable needs a finer grained control over the lock, there's no higher up lock here.
+
 	bool bCanAddObject =
 		ensureMsgf(Storage, TEXT("Trying to add a UObject to Typed Element's Data Storage before the storage is available.")) &&
 		ShouldAddObject(Object);
 	return bCanAddObject ? AddCompatibleObjectExplicitTransactionable<true>(Object) : TypedElementDataStorage::InvalidRowHandle;
 }
 
-TypedElementRowHandle UTypedElementDatabaseCompatibility::AddCompatibleObjectExplicit(void* Object, TWeakObjectPtr<const UScriptStruct> TypeInfo)
+TypedElementRowHandle UTypedElementDatabaseCompatibility::AddCompatibleObjectExplicit(void* Object, TWeakObjectPtr<UScriptStruct> TypeInfo)
 {
 	using namespace TypedElementDataStorage;
+	using namespace UE::EditorDataStorage;
 	
-	if (ensureMsgf(Storage, TEXT("Trying to add an object to Typed Element's Data Storage before the storage is available.")))
+	checkf(Storage, TEXT("Trying to add an object to Typed Element's Data Storage before the storage is available."));
+	FScopedExclusiveLock Lock(EGlobalLockScope::Public);
+	
+	RowHandle Result = FindRowWithCompatibleObjectExplicit(Object);
+	if (!Storage->IsRowAvailable(Result))
 	{
-		TypedElementRowHandle Result = FindRowWithCompatibleObjectExplicit(Object);
-		if (!Storage->IsRowAvailable(Result))
+		Result = Storage->ReserveRow();
+		Storage->IndexRow(GenerateIndexHash(Object), Result);
+		if (bUseCommandBuffer)
+		{
+			QueuedCommands.AddCommand(FAddCompatibleExternalObject{ .Object = Object, .TypeInfo = TypeInfo, .Row = Result });
+		}
+		else
 		{
-			Result = Storage->ReserveRow();
-			Storage->IndexRow(GenerateIndexHash(Object), Result);
 			ExternalObjectsPendingRegistration.Add(Result, ExternalObjectRegistration{ .Object = Object, .TypeInfo = TypeInfo });
 		}
-		return Result;
-	}
-	else
-	{
-		return TypedElementInvalidRowHandle;
 	}
+	return Result;
 }
 
 void UTypedElementDatabaseCompatibility::RemoveCompatibleObjectExplicit(UObject* Object)
@@ -159,35 +252,51 @@ void UTypedElementDatabaseCompatibility::RemoveCompatibleObjectExplicit(UObject*
 void UTypedElementDatabaseCompatibility::RemoveCompatibleObjectExplicit(void* Object)
 {
 	using namespace TypedElementDataStorage;
+	using namespace UE::EditorDataStorage;
 
 	checkf(Storage, TEXT("Removing compatible objects is not supported before Typed Element's Database compatibility manager has been initialized."));
-	IndexHash Hash = GenerateIndexHash(Object);
-	RowHandle Row = Storage->FindIndexedRow(Hash);
-	if (Storage->IsRowAvailable(Row))
+
+	if (TypedElementDataStorage::bUseCommandBuffer && TypedElementDataStorage::bUseDeferredRemovesInCompat)
+	{
+		QueuedCommands.AddCommand(FRemoveCompatibleExternalObject{ .Object = Object });
+	}
+	else
 	{
-		const FTypedElementScriptStructTypeInfoColumn* TypeInfoColumn = Storage->GetColumn<FTypedElementScriptStructTypeInfoColumn>(Row);
-		if (Storage->IsRowAssigned(Row) && ensureMsgf(TypeInfoColumn, TEXT("Missing type information for removed void* object at ptr 0x%p"), Object))
+		FScopedExclusiveLock Lock(EGlobalLockScope::Public);
+
+		IndexHash Hash = GenerateIndexHash(Object);
+		RowHandle Row = Storage->FindIndexedRow(Hash);
+		if (Storage->IsRowAvailable(Row))
 		{
-			OnPreObjectRemoved(Object, TypeInfoColumn->TypeInfo.Get(), Row);
+			const FTypedElementScriptStructTypeInfoColumn* TypeInfoColumn = Storage->GetColumn<FTypedElementScriptStructTypeInfoColumn>(Row);
+			if (Storage->IsRowAssigned(Row) && ensureMsgf(TypeInfoColumn, TEXT("Missing type information for removed void* object at ptr 0x%p"), Object))
+			{
+				TriggerOnPreObjectRemoved(Object, TypeInfoColumn->TypeInfo.Get(), Row);
+			}
+			Storage->RemoveRow(Row);
 		}
-		Storage->RemoveRow(Row);
 	}
 }
 
 TypedElementRowHandle UTypedElementDatabaseCompatibility::FindRowWithCompatibleObjectExplicit(const UObject* Object) const
 {
 	using namespace TypedElementDataStorage;
+	using namespace UE::EditorDataStorage;
 
 	if (Object && Storage && Storage->IsAvailable())
 	{
+		FScopedSharedLock Lock(EGlobalLockScope::Public);
+
 		RowHandle Row = Storage->FindIndexedRow(GenerateIndexHash(Object));
 		return Storage->IsRowAvailable(Row) ? Row : DealiasObject(Object);
 	}
-	return TypedElementInvalidRowHandle;
+	return InvalidRowHandle;
 }
 
 TypedElementRowHandle UTypedElementDatabaseCompatibility::FindRowWithCompatibleObjectExplicit(const void* Object) const
 {
+	// Thread safety is only needed by FindIndexedRow which internally takes care of it.
+
 	using namespace TypedElementDataStorage;
 
 	return (Object && Storage && Storage->IsAvailable()) ? Storage->FindIndexedRow(GenerateIndexHash(Object)) : InvalidRowHandle;
@@ -195,12 +304,18 @@ TypedElementRowHandle UTypedElementDatabaseCompatibility::FindRowWithCompatibleO
 
 bool UTypedElementDatabaseCompatibility::SupportsExtension(FName Extension) const
 {
+	// No thread safety needed.
+
 	using namespace TypedElementDataStorage;
 
 	if (Extension == IntegrateWithGCName)
 	{
 		return bIntegrateWithGC;
 	}
+	else if (Extension == CompatibilityUsesCommandBufferExtensionName)
+	{
+		return bUseCommandBuffer;
+	}
 	else
 	{
 		return false;
@@ -209,16 +324,23 @@ bool UTypedElementDatabaseCompatibility::SupportsExtension(FName Extension) cons
 
 void UTypedElementDatabaseCompatibility::ListExtensions(TFunctionRef<void(FName)> Callback) const
 {
+	// No thread safety needed.
+
 	using namespace TypedElementDataStorage;
 
 	if (bIntegrateWithGC)
 	{
 		Callback(IntegrateWithGCName);
 	}
+	if (bUseCommandBuffer)
+	{
+		Callback(CompatibilityUsesCommandBufferExtensionName);
+	}
 }
 
 void UTypedElementDatabaseCompatibility::Prepare()
 {
+	// Thread-safe as this is only called from a function that has an exclusive lock.
 	CreateStandardArchetypes();
 	RegisterTypeInformationQueries();
 }
@@ -229,6 +351,8 @@ void UTypedElementDatabaseCompatibility::Reset()
 
 void UTypedElementDatabaseCompatibility::CreateStandardArchetypes()
 {
+	// Thread-safe as this is only called from a function that has an exclusive lock.
+
 	StandardActorTable = Storage->RegisterTable(TTypedElementColumnTypeList<
 			FTypedElementUObjectColumn, FTypedElementUObjectIdColumn, FTypedElementClassTypeInfoColumn,
 			FTypedElementLabelColumn, FTypedElementLabelHashColumn, FTypedElementActorTag,
@@ -255,6 +379,8 @@ void UTypedElementDatabaseCompatibility::CreateStandardArchetypes()
 
 void UTypedElementDatabaseCompatibility::RegisterTypeInformationQueries()
 {
+	// Thread-safe as this is only called from a function that has an exclusive lock.
+
 	using namespace TypedElementQueryBuilder;
 
 	ClassTypeInfoQuery = Storage->RegisterQuery(
@@ -276,6 +402,9 @@ void UTypedElementDatabaseCompatibility::RegisterTypeInformationQueries()
 bool UTypedElementDatabaseCompatibility::ShouldAddObject(const UObject* Object) const
 {
 	using namespace TypedElementDataStorage;
+	using namespace UE::EditorDataStorage;
+	
+	FScopedSharedLock Lock(EGlobalLockScope::Public);
 
 	bool Include = true;
 	if (!Storage->IsRowAvailable(Storage->FindIndexedRow(GenerateIndexHash(Object))))
@@ -293,6 +422,9 @@ bool UTypedElementDatabaseCompatibility::ShouldAddObject(const UObject* Object)
 TypedElementDataStorage::TableHandle UTypedElementDatabaseCompatibility::FindBestMatchingTable(const UStruct* TypeInfo) const
 {
 	using namespace TypedElementDataStorage;
+	using namespace UE::EditorDataStorage;
+
+	FScopedSharedLock Lock(EGlobalLockScope::Public);
 
 	while (TypeInfo)
 	{
@@ -309,17 +441,25 @@ TypedElementDataStorage::TableHandle UTypedElementDatabaseCompatibility::FindBes
 template<bool bEnableTransactions>
 TypedElementRowHandle UTypedElementDatabaseCompatibility::AddCompatibleObjectExplicitTransactionable(UObject* Object)
 {
-	check(IsInGameThread());
-
 	using namespace TypedElementDataStorage;
+	using namespace UE::EditorDataStorage;
+
+	FScopedExclusiveLock Lock(EGlobalLockScope::Public);
 
 	TypedElementRowHandle Result = FindRowWithCompatibleObjectExplicit(Object);
 	if (!Storage->IsRowAvailable(Result))
 	{
 		Result = Storage->ReserveRow();
 		Storage->IndexRow(GenerateIndexHash(Object), Result);
-		UObjectsPendingRegistration.Add(Result, Object);
-
+		if (bUseCommandBuffer)
+		{
+			QueuedCommands.AddCommand(FAddCompatibleUObject{ .Object = Object, .Row = Result });
+		}
+		else
+		{
+			UObjectsPendingRegistration.Add(Result, Object);
+		}
+		
 		if constexpr (bEnableTransactions)
 		{
 			if (IsInGameThread() && GUndo)
@@ -328,7 +468,6 @@ TypedElementRowHandle UTypedElementDatabaseCompatibility::AddCompatibleObjectExp
 			}
 		}
 	}
-
 	return Result;
 }
 
@@ -336,9 +475,25 @@ template<bool bEnableTransactions>
 void UTypedElementDatabaseCompatibility::RemoveCompatibleObjectExplicitTransactionable(const UObject* Object)
 {
 	using namespace TypedElementDataStorage;
+	using namespace UE::EditorDataStorage;
 
 	checkf(Storage,
 		TEXT("Removing compatible objects is not supported before Typed Element's Database compatibility manager has been initialized."));
+	
+	if constexpr (!bEnableTransactions)
+	{
+		if (TypedElementDataStorage::bUseCommandBuffer && TypedElementDataStorage::bUseDeferredRemovesInCompat)
+		{
+			// There's no need for an transaction recording so the full operation can be done as part of the commands processing.
+			QueuedCommands.AddCommand(FRemoveCompatibleUObject{ .Object = Object });
+			return;
+		}
+	}
+
+	// Do not lock while both buffered and non-buffered ways are still available. An exclusive lock is required here for the
+	// non-buffered to reduce the additional locks/unlocks while the buffered version doesn't need any locking beyond the
+	// shared lock FindIndexedRow does. Not adding an exclusive here means some additional lock/unlocking but doesn't make
+	// the code thread unsafe.
 	IndexHash Hash = GenerateIndexHash(Object);
 	RowHandle Row = Storage->FindIndexedRow(Hash);
 
@@ -353,16 +508,13 @@ void UTypedElementDatabaseCompatibility::RemoveCompatibleObjectExplicitTransacti
 	const UObject* Object, TypedElementDataStorage::RowHandle ObjectRow)
 {
 	using namespace TypedElementDataStorage;
+	using namespace UE::EditorDataStorage;
 
-	checkf(Storage, 
+	checkf(Storage,
 		TEXT("Removing compatible objects is not supported before Typed Element's Database compatibility manager has been initialized."));
-	
-	const FTypedElementClassTypeInfoColumn* TypeInfoColumn = Storage->GetColumn<FTypedElementClassTypeInfoColumn>(ObjectRow);
-	if (Storage->IsRowAssigned(ObjectRow) &&
-		ensureMsgf(TypeInfoColumn, TEXT("Missing type information for removed UObject at ptr 0x%p [%s]"), Object, *Object->GetName()))
-	{
-		OnPreObjectRemoved(Object, TypeInfoColumn->TypeInfo.Get(), ObjectRow);
 
+	if (TypedElementDataStorage::bUseCommandBuffer && TypedElementDataStorage::bUseDeferredRemovesInCompat)
+	{
 		if constexpr (bEnableTransactions)
 		{
 			if (IsInGameThread() && GUndo)
@@ -370,13 +522,35 @@ void UTypedElementDatabaseCompatibility::RemoveCompatibleObjectExplicitTransacti
 				GUndo->StoreUndo(this, MakeUnique<FDeregistrationCommandChange>(this, const_cast<UObject*>(Object)));
 			}
 		}
+		QueuedCommands.AddCommand(FRemoveCompatibleUObject{ .Object = Object, .ObjectRow = ObjectRow });
 	}
+	else
+	{
+		FScopedExclusiveLock Lock(EGlobalLockScope::Public);
+
+		const FTypedElementClassTypeInfoColumn* TypeInfoColumn = Storage->GetColumn<FTypedElementClassTypeInfoColumn>(ObjectRow);
+		if (Storage->IsRowAssigned(ObjectRow) &&
+			ensureMsgf(TypeInfoColumn, TEXT("Missing type information for removed UObject at ptr 0x%p [%s]"), Object, *Object->GetName()))
+		{
+			TriggerOnPreObjectRemoved(Object, TypeInfoColumn->TypeInfo.Get(), ObjectRow);
+
+			if constexpr (bEnableTransactions)
+			{
+				if (IsInGameThread() && GUndo)
+				{
+					GUndo->StoreUndo(this, MakeUnique<FDeregistrationCommandChange>(this, const_cast<UObject*>(Object)));
+				}
+			}
+		}
 
-	Storage->RemoveRow(ObjectRow);
+		Storage->RemoveRow(ObjectRow);
+	}
 }
 
 TypedElementRowHandle UTypedElementDatabaseCompatibility::DealiasObject(const UObject* Object) const
 {
+	// Thread safe because it's only called from functions that already lock.
+
 	for (const ObjectToRowDealiaser& Dealiaser : ObjectToRowDialiasers)
 	{
 		if (TypedElementRowHandle Row = Dealiaser(*this, Object); Storage->IsRowAvailable(Row))
@@ -389,16 +563,26 @@ TypedElementRowHandle UTypedElementDatabaseCompatibility::DealiasObject(const UO
 
 void UTypedElementDatabaseCompatibility::Tick()
 {
+	using namespace UE::EditorDataStorage;
+	
 	TEDS_EVENT_SCOPE(TEXT("Compatibility Tick"))
 	
-	PendingTypeInformationUpdate.Process(*this);
+	FScopedExclusiveLock Lock(EGlobalLockScope::Public);
 
 	// Delay processing until the required systems are available by not clearing any lists or doing any work.
 	if (Storage && Storage->IsAvailable())
 	{
-		TickPendingUObjectRegistration();
-		TickPendingExternalObjectRegistration();
-		TickObjectSync();
+		if (TypedElementDataStorage::bUseCommandBuffer)
+		{
+			TickPendingCommands();
+		}
+		else
+		{
+			PendingTypeInformationUpdate.Process(*this);
+			TickPendingUObjectRegistration();
+			TickPendingExternalObjectRegistration();
+			TickObjectSync();
+		}
 	}
 }
 
@@ -431,15 +615,19 @@ void UTypedElementDatabaseCompatibility::FPendingTypeInformationUpdate::Process(
 {
 	using namespace TypedElementDataStorage;
 	using namespace TypedElementQueryBuilder;
-	
+	using namespace UE::EditorDataStorage;
+
 	if (bHasPendingUpdate)
 	{
 		// Swap to release the lock as soon as possible.
 		{
 			UE::TUniqueLock Lock(Safeguard);
 			std::swap(PendingTypeInformationUpdatesActive, PendingTypeInformationUpdatesSwapped);
+			bHasPendingUpdate = false;
 		}
 
+		FScopedExclusiveLock Lock(EGlobalLockScope::Public);
+
 		for (TypeToTableMapType::TIterator It = Compatibility.TypeToTableMap.CreateIterator(); It; ++It)
 		{
 			if (TOptional<TWeakObjectPtr<UObject>> NewObject = ProcessResolveTypeRecursively(It.Key()); NewObject.IsSet())
@@ -459,7 +647,10 @@ void UTypedElementDatabaseCompatibility::FPendingTypeInformationUpdate::Process(
 		Compatibility.Storage->RunQuery(Compatibility.ClassTypeInfoQuery, CreateDirectQueryCallbackBinding(
 			[this](IDirectQueryContext& Context, FTypedElementClassTypeInfoColumn& Type)
 			{
-				if (TOptional<TWeakObjectPtr<UObject>> NewObject = ProcessResolveTypeRecursively(Type.TypeInfo); NewObject.IsSet())
+				if (TOptional<TWeakObjectPtr<UObject>> NewObject = 
+					// Using reintpret_Cast as TWeakObjectPtr<const UClass> and TWeakObjectPtr<UClass> are considered separate types by 
+					// the compiler.
+					ProcessResolveTypeRecursively(reinterpret_cast<TWeakObjectPtr<UClass>&>(Type.TypeInfo)); NewObject.IsSet())
 				{
 					Type.TypeInfo = Cast<UClass>(*NewObject);
 					checkf(Type.TypeInfo.IsValid(),
@@ -469,7 +660,8 @@ void UTypedElementDatabaseCompatibility::FPendingTypeInformationUpdate::Process(
 		Compatibility.Storage->RunQuery(Compatibility.ScriptStructTypeInfoQuery, CreateDirectQueryCallbackBinding(
 			[this](IDirectQueryContext& Context, FTypedElementScriptStructTypeInfoColumn& Type)
 			{
-				if (TOptional<TWeakObjectPtr<UObject>> NewObject = ProcessResolveTypeRecursively(Type.TypeInfo); NewObject.IsSet())
+				if (TOptional<TWeakObjectPtr<UObject>> NewObject = 
+					ProcessResolveTypeRecursively(reinterpret_cast<TWeakObjectPtr<UScriptStruct>&>(Type.TypeInfo)); NewObject.IsSet())
 				{
 					Type.TypeInfo = Cast<UScriptStruct>(*NewObject);
 					checkf(Type.TypeInfo.IsValid(),
@@ -489,13 +681,14 @@ void UTypedElementDatabaseCompatibility::FPendingTypeInformationUpdate::Process(
 			});
 
 		PendingTypeInformationUpdatesSwapped->Reset();
-		bHasPendingUpdate = false;
 	}
 }
 
 TOptional<TWeakObjectPtr<UObject>> UTypedElementDatabaseCompatibility::FPendingTypeInformationUpdate::ProcessResolveTypeRecursively(
-	const TWeakObjectPtr<const UObject>& Target)
+	const TWeakObjectPtr<UObject>& Target)
 {
+	// Thread-safety guaranteed because this is a private function that only gets called from functions that called inside a mutex.
+
 	if (const TWeakObjectPtr<UObject>* NewObject = PendingTypeInformationUpdatesSwapped->Find(Target))
 	{
 		TWeakObjectPtr<UObject> LastNewObject = *NewObject;
@@ -518,24 +711,28 @@ TOptional<TWeakObjectPtr<UObject>> UTypedElementDatabaseCompatibility::FPendingT
 template<typename AddressType>
 void UTypedElementDatabaseCompatibility::PendingRegistration<AddressType>::Add(TypedElementRowHandle ReservedRowHandle, AddressType Address)
 {
+	// Thread-safe as it's only called from functions that already lock using.
 	Entries.Emplace(FEntry{ .Address = Address, .Row = ReservedRowHandle });
 }
 
 template<typename AddressType>
 bool UTypedElementDatabaseCompatibility::PendingRegistration<AddressType>::IsEmpty() const
 {
+	// Thread-safe as it's only called from functions that already lock using.
 	return Entries.IsEmpty();
 }
 
 template<typename AddressType>
 int32 UTypedElementDatabaseCompatibility::PendingRegistration<AddressType>::Num() const
 {
+	// Thread-safe as it's only called from functions that already lock using.
 	return Entries.Num();
 }
 
 template<typename AddressType>
 void UTypedElementDatabaseCompatibility::PendingRegistration<AddressType>::ForEachAddress(const TFunctionRef<void(AddressType&)>& Callback)
 {
+	// Thread-safe as it's only called from functions that already lock using.
 	for (FEntry& Entry : Entries)
 	{
 		Callback(Entry.Address);
@@ -546,6 +743,7 @@ template<typename AddressType>
 void UTypedElementDatabaseCompatibility::PendingRegistration<AddressType>::ProcessEntries(ITypedElementDataStorageInterface& StorageInterface,
 	UTypedElementDatabaseCompatibility& Compatibility, const TFunctionRef<void(TypedElementRowHandle, const AddressType&)>& SetupRowCallback)
 {
+	// Thread-safe as it's only called from functions that already lock using.
 	using namespace TypedElementDataStorage;
 
 	// Start by removing any entries that are no longer valid.
@@ -652,11 +850,61 @@ void UTypedElementDatabaseCompatibility::PendingRegistration<AddressType>::Proce
 template<typename AddressType>
 void UTypedElementDatabaseCompatibility::PendingRegistration<AddressType>::Reset()
 {
+	// Thread-safe as it's only called from functions that already lock using.
 	Entries.Reset();
 }
 
+void UTypedElementDatabaseCompatibility::TickPendingCommands()
+{
+	using namespace UE::EditorDataStorage;
+	using namespace TypedElementDataStorage;
+
+	// Thread safe because it's only called from functions that already lock.
+	SIZE_T CommandCount = QueuedCommands.Collect(PendingCommands);
+
+	// First see if there's anything that needs to be patched to avoid any of the later steps using stale data.
+	if (FPatchData::IsPatchingRequired(PendingCommands))
+	{
+		TEDS_EVENT_SCOPE(TEXT("Compatibility Tick - Patching"))
+		FPatchData::RunPatch(PendingCommands, *this, Environment->GetScratchBuffer());
+		CommandCount = PendingCommands.GetTotalCommandCount();
+	}
+
+	if (CommandCount > 0)
+	{
+		TEDS_EVENT_SCOPE(TEXT("Compatibility Tick - Preparation"))
+		// Prepare data in the commands. Commands that can't or don't need to be executed will be nop-ed out.
+		FPrepareCommands::RunPreparation(*Storage, *this, PendingCommands);
+		CommandCount = PendingCommands.GetTotalCommandCount();
+	}
+
+	if (CommandCount > 0)
+	{
+		if (bOptimizeCommandBuffer)
+		{
+			TEDS_EVENT_SCOPE(TEXT("Compatibility Tick - Optimization"))
+			FSorter::SortCommands(PendingCommands);
+			FCommandOptimizer::Run(PendingCommands, Environment->GetScratchBuffer());
+		}
+
+		if (PrintCompatCommandBuffer > 0 )
+		{
+			TEDS_EVENT_SCOPE(TEXT("Compatibility Tick - Logging"))
+			FString CommandsAsString = FRecordCommands::PrintToString(PendingCommands, PrintCompatCommandBuffer == 2);
+			UE_LOG(LogTedsCompat, Log, TEXT("Pending Commands:\n%s%u Nops"), 
+				*CommandsAsString, PendingCommands.GetCommandCount<FNopCommand>());
+		}
+
+		TEDS_EVENT_SCOPE(TEXT("Compatibility Tick - Processing"))
+		PendingCommands.Process(FCommandProcessor(*Storage, *this));
+	}
+	PendingCommands.Reset();
+}
+
 void UTypedElementDatabaseCompatibility::TickPendingUObjectRegistration()
 {
+	// Thread safe because it's only called from functions that already lock.
+
 	if (!UObjectsPendingRegistration.IsEmpty())
 	{
 		UObjectsPendingRegistration.ProcessEntries(*Storage, *this,
@@ -676,7 +924,7 @@ void UTypedElementDatabaseCompatibility::TickPendingUObjectRegistration()
 				}
 				// Make sure the new row is tagged for update.
 				Interface->AddColumn<FTypedElementSyncFromWorldTag>(Row);
-				OnObjectAdded(Object.Get(), Object->GetClass(), Row);
+				TriggerOnObjectAdded(Object.Get(), Object->GetClass(), Row);
 			});
 
 		UObjectsPendingRegistration.Reset();
@@ -685,6 +933,8 @@ void UTypedElementDatabaseCompatibility::TickPendingUObjectRegistration()
 
 void UTypedElementDatabaseCompatibility::TickPendingExternalObjectRegistration()
 {
+	// Thread safe because it's only called from functions that already lock.
+
 	if (!ExternalObjectsPendingRegistration.IsEmpty())
 	{
 		ExternalObjectsPendingRegistration.ProcessEntries(*Storage, *this,
@@ -695,7 +945,7 @@ void UTypedElementDatabaseCompatibility::TickPendingExternalObjectRegistration()
 				Interface->AddColumn(Row, FTypedElementScriptStructTypeInfoColumn{ .TypeInfo = Object.TypeInfo });
 				// Make sure the new row is tagged for update.
 				Interface->AddColumn<FTypedElementSyncFromWorldTag>(Row);
-				OnObjectAdded(Object.Object, Object.TypeInfo.Get(), Row);
+				TriggerOnObjectAdded(Object.Object, Object.TypeInfo.Get(), Row);
 			});
 
 		ExternalObjectsPendingRegistration.Reset();
@@ -706,6 +956,8 @@ void UTypedElementDatabaseCompatibility::TickObjectSync()
 {
 	using namespace TypedElementDataStorage;
 	
+	// Thread safe because it's only called from functions that already lock.
+
 	if (!ObjectsNeedingSyncTags.IsEmpty())
 	{
 		TEDS_EVENT_SCOPE(TEXT("Process ObjectsNeedingSyncTags"));
@@ -746,36 +998,76 @@ void UTypedElementDatabaseCompatibility::TickObjectSync()
 
 void UTypedElementDatabaseCompatibility::OnPrePropertyChanged(UObject* Object, const FEditPropertyChain& PropertyChain)
 {
-	ObjectsNeedingSyncTags.FindOrAdd(Object).AddUnique(
-		FSyncTagInfo{ .ColumnType = FTypedElementSyncFromWorldInteractiveTag::StaticStruct(), .bAddColumn = true });
+	using namespace UE::EditorDataStorage;
+
+	if (TypedElementDataStorage::bUseCommandBuffer)
+	{
+		QueuedCommands.AddCommand(FAddInteractiveSyncFromWorldTag{ .Target = Object });
+	}
+	else
+	{
+		FScopedExclusiveLock Lock(EGlobalLockScope::Public);
+		ObjectsNeedingSyncTags.FindOrAdd(Object).AddUnique(
+			FSyncTagInfo{ .ColumnType = FTypedElementSyncFromWorldInteractiveTag::StaticStruct(), .bAddColumn = true });
+	}
 }
 
 void UTypedElementDatabaseCompatibility::OnPostEditChangeProperty(
 	UObject* Object,
 	FPropertyChangedEvent& PropertyChangedEvent)
 {
-	// Determining the object is being tracked in the database can't be done safely as it may be queued for addition.
-	// It would also add a small bit of performance overhead as access the lookup table can be done faster as a
-	// batch operation during the tick step.
-	if (PropertyChangedEvent.ChangeType != EPropertyChangeType::Interactive)
+	using namespace UE::EditorDataStorage;
+
+	if (TypedElementDataStorage::bUseCommandBuffer)
 	{
-		ObjectsNeedingSyncTagsMapValue& SyncValue = ObjectsNeedingSyncTags.FindOrAdd(Object);
-		SyncValue.AddUnique(FSyncTagInfo{ .ColumnType = FTypedElementSyncFromWorldTag::StaticStruct(), .bAddColumn = true });
-		SyncValue.AddUnique(FSyncTagInfo{ .ColumnType = FTypedElementSyncFromWorldInteractiveTag::StaticStruct(), .bAddColumn = false });
+		if (PropertyChangedEvent.ChangeType != EPropertyChangeType::Interactive)
+		{
+			QueuedCommands.AddCommand(FRemoveInteractiveSyncFromWorldTag{ .Target = Object });
+		}
+	}
+	else
+	{
+		FScopedExclusiveLock Lock(EGlobalLockScope::Public);
+
+		// Determining the object is being tracked in the database can't be done safely as it may be queued for addition.
+		// It would also add a small bit of performance overhead as access the lookup table can be done faster as a
+		// batch operation during the tick step.
+		if (PropertyChangedEvent.ChangeType != EPropertyChangeType::Interactive)
+		{
+			ObjectsNeedingSyncTagsMapValue& SyncValue = ObjectsNeedingSyncTags.FindOrAdd(Object);
+			SyncValue.AddUnique(FSyncTagInfo{ .ColumnType = FTypedElementSyncFromWorldTag::StaticStruct(), .bAddColumn = true });
+			SyncValue.AddUnique(FSyncTagInfo{ .ColumnType = FTypedElementSyncFromWorldInteractiveTag::StaticStruct(), .bAddColumn = false });
+		}
 	}
 }
 
 void UTypedElementDatabaseCompatibility::OnObjectModified(UObject* Object)
 {
-	// Determining the object is being tracked in the database can't be done safely as it may be queued for addition.
-	// It would also add a small bit of performance overhead as access the lookup table can be done faster as a
-	// batch operation during the tick step.
-	ObjectsNeedingSyncTags.FindOrAdd(Object).AddUnique(
-		FSyncTagInfo{ .ColumnType = FTypedElementSyncFromWorldTag::StaticStruct(), .bAddColumn = true });
+	using namespace UE::EditorDataStorage;
+
+	if (TypedElementDataStorage::bUseCommandBuffer)
+	{
+		QueuedCommands.AddCommand(FAddSyncFromWorldTag{ .Target = Object });
+	}
+	else
+	{
+		FScopedExclusiveLock Lock(EGlobalLockScope::Public);
+
+		// Determining the object is being tracked in the database can't be done safely as it may be queued for addition.
+		// It would also add a small bit of performance overhead as access the lookup table can be done faster as a
+		// batch operation during the tick step.
+		ObjectsNeedingSyncTags.FindOrAdd(Object).AddUnique(
+			FSyncTagInfo{ .ColumnType = FTypedElementSyncFromWorldTag::StaticStruct(), .bAddColumn = true });
+	}
 }
 
-void UTypedElementDatabaseCompatibility::OnObjectAdded(const void* Object, FTypedElementDatabaseCompatibilityObjectTypeInfo TypeInfo, TypedElementRowHandle Row) const
+void UTypedElementDatabaseCompatibility::TriggerOnObjectAdded(
+	const void* Object, UE::EditorDataStorage::FObjectTypeInfo TypeInfo, TypedElementDataStorage::RowHandle Row) const
 {
+	using namespace UE::EditorDataStorage;
+
+	// Thread safe because it's only called from functions that already lock.
+
 	for (const TPair<ObjectAddedCallback, FDelegateHandle>& CallbackPair : ObjectAddedCallbackList)
 	{
 		const ObjectAddedCallback& Callback = CallbackPair.Key;
@@ -783,8 +1075,13 @@ void UTypedElementDatabaseCompatibility::OnObjectAdded(const void* Object, FType
 	}
 }
 
-void UTypedElementDatabaseCompatibility::OnPreObjectRemoved(const void* Object, FTypedElementDatabaseCompatibilityObjectTypeInfo TypeInfo, TypedElementRowHandle Row) const
+void UTypedElementDatabaseCompatibility::TriggerOnPreObjectRemoved(
+	const void* Object, UE::EditorDataStorage::FObjectTypeInfo TypeInfo, TypedElementDataStorage::RowHandle Row) const
 {
+	using namespace UE::EditorDataStorage;
+
+	// Thread safe because it's only called from functions that already lock.
+
 	for (const TPair<ObjectRemovedCallback, FDelegateHandle>& CallbackPair : PreObjectRemovedCallbackList)
 	{
 		const ObjectRemovedCallback& Callback = CallbackPair.Key;
@@ -794,44 +1091,98 @@ void UTypedElementDatabaseCompatibility::OnPreObjectRemoved(const void* Object,
 
 void UTypedElementDatabaseCompatibility::OnObjectReinstanced(const FCoreUObjectDelegates::FReplacementObjectMap& ReplacedObjects)
 {
-	PendingTypeInformationUpdate.AddTypeInformation(ReplacedObjects);
+	using namespace UE::EditorDataStorage;
+	
+	if (TypedElementDataStorage::bUseCommandBuffer)
+	{
+		bool bHasUpdatedTypeInformation = false;
+		for (TMap<UObject*, UObject*>::TConstIterator It = ReplacedObjects.CreateConstIterator(); It; ++It)
+		{
+			UStruct* Original = Cast<UStruct>(It->Key);
+			UStruct* Reinstanced = Cast<UStruct>(It->Value);
+			if (Original && Reinstanced)
+			{
+				QueuedCommands.AddCommand(FTypeInfoReinstanced{ .Original = Original, .Reinstanced = Reinstanced });
+				bHasUpdatedTypeInformation = true;
+			}
+		}
+	}
+	else
+	{
+		FScopedExclusiveLock Lock(EGlobalLockScope::Public);
+		PendingTypeInformationUpdate.AddTypeInformation(ReplacedObjects);
+	}
 }
 
 void UTypedElementDatabaseCompatibility::OnPostGcUnreachableAnalysis()
 {
 	using namespace TypedElementDataStorage;
 	using namespace TypedElementQueryBuilder;
+	using namespace UE::EditorDataStorage;
 
 	if (bIntegrateWithGC)
 	{
-		TArray<TPair<FUObjectItem*, RowHandle>> DeletedObjects;
+		TEDS_EVENT_SCOPE(TEXT("Post GC clean up"));
+		FScopedExclusiveLock Lock(EGlobalLockScope::Public);
 
-		TEDS_EVENT_SCOPE(TEXT("Post GC clean up"))
-
-		Storage->RunQuery(UObjectQuery, CreateDirectQueryCallbackBinding(
-			[&DeletedObjects](RowHandle Row, const FTypedElementUObjectIdColumn& ObjectId)
-			{
-				FUObjectItem* Description = GUObjectArray.IndexToObject(ObjectId.Id);
-				if (ensureMsgf(Description && Description->SerialNumber == ObjectId.SerialNumber, 
-					TEXT("The UObject found in TEDS no longer exists. TEDS was likely not informed in an earlier GC pass.")))
-					// Unable to provide additional information such as the UObject's name as the UObject will not be valid.
+		if (TypedElementDataStorage::bUseCommandBuffer)
+		{
+			Storage->RunQuery(UObjectQuery, CreateDirectQueryCallbackBinding(
+				[this](RowHandle Row, const FTypedElementUObjectIdColumn& ObjectId)
 				{
-					if (Description->HasAnyFlags(EInternalObjectFlags::Garbage | EInternalObjectFlags::Unreachable))
+					FUObjectItem* Description = GUObjectArray.IndexToObject(ObjectId.Id);
+					if (ensureMsgf(Description && Description->SerialNumber == ObjectId.SerialNumber,
+						TEXT("The UObject found in TEDS no longer exists. TEDS was likely not informed in an earlier GC pass.")))
+						// Unable to provide additional information such as the UObject's name as the UObject will not be valid.
 					{
-						DeletedObjects.Emplace(Description, Row);
+						if (Description->HasAnyFlags(EInternalObjectFlags::Garbage | EInternalObjectFlags::Unreachable))
+						{
+							if (UObject* Object = Cast<UObject>(Description->Object)) // No need to delete if this isn't a full UObject.
+							{
+								QueuedCommands.AddCommand(FRemoveCompatibleUObject{ .Object = Object, .ObjectRow = Row });
+							}
+						}
 					}
-				}
-			}));
-
-		for (TPair<FUObjectItem*, RowHandle>& Object : DeletedObjects)
+				}));
+			// Forcefully execute all pending commands to make sure there are no commands left that reference deleted objects as well
+			// as to make sure the added deletes are executed to guaranteed there are no stale objects in TEDS.
+			TickPendingCommands();
+		}
+		else
 		{
-			RemoveCompatibleObjectExplicitTransactionable<false>(Cast<UObject>(Object.Key->Object), Object.Value);
+			TArray<TPair<FUObjectItem*, RowHandle>> DeletedObjects;
+			Storage->RunQuery(UObjectQuery, CreateDirectQueryCallbackBinding(
+				[&DeletedObjects](RowHandle Row, const FTypedElementUObjectIdColumn& ObjectId)
+				{
+					FUObjectItem* Description = GUObjectArray.IndexToObject(ObjectId.Id);
+					if (ensureMsgf(Description && Description->SerialNumber == ObjectId.SerialNumber,
+						TEXT("The UObject found in TEDS no longer exists. TEDS was likely not informed in an earlier GC pass.")))
+						// Unable to provide additional information such as the UObject's name as the UObject will not be valid.
+					{
+						if (Description->HasAnyFlags(EInternalObjectFlags::Garbage | EInternalObjectFlags::Unreachable))
+						{
+							DeletedObjects.Emplace(Description, Row);
+						}
+					}
+				}));
+
+			for (TPair<FUObjectItem*, RowHandle>& ObjectItem : DeletedObjects)
+			{
+				if (UObject* Object = Cast<UObject>(ObjectItem.Key->Object)) // No need to delete if this isn't a full UObject.
+				{
+					RemoveCompatibleObjectExplicitTransactionable<false>(Object, ObjectItem.Value);
+				}
+			}
 		}
 	}
 }
 
 void UTypedElementDatabaseCompatibility::OnPostWorldInitialization(UWorld* World, const UWorld::InitializationValues InitializationValues)
 {
+	using namespace UE::EditorDataStorage;
+
+	FScopedExclusiveLock Lock(EGlobalLockScope::Public);
+
 	FDelegateHandle Handle = World->AddOnActorDestroyedHandler(
 		FOnActorDestroyed::FDelegate::CreateUObject(this, &UTypedElementDatabaseCompatibility::OnActorDestroyed));
 	ActorDestroyedDelegateHandles.Add(World, Handle);
@@ -839,6 +1190,10 @@ void UTypedElementDatabaseCompatibility::OnPostWorldInitialization(UWorld* World
 
 void UTypedElementDatabaseCompatibility::OnPreWorldFinishDestroy(UWorld* World)
 {
+	using namespace UE::EditorDataStorage;
+
+	FScopedExclusiveLock Lock(EGlobalLockScope::Public);
+
 	FDelegateHandle Handle;
 	if (ActorDestroyedDelegateHandles.RemoveAndCopyValue(World, Handle))
 	{
@@ -848,11 +1203,13 @@ void UTypedElementDatabaseCompatibility::OnPreWorldFinishDestroy(UWorld* World)
 
 void UTypedElementDatabaseCompatibility::OnActorDestroyed(AActor* Actor)
 {
+	// The only function called is already thread safe.
 	RemoveCompatibleObjectExplicit(Actor);
 }
 
 SIZE_T GetTypeHash(const UTypedElementDatabaseCompatibility::FSyncTagInfo& Column)
 {
+	// Thread safe as it only uses local data.
 	return HashCombine(Column.ColumnType.GetWeakPtrTypeHash(), Column.bAddColumn);
 }
 
@@ -869,35 +1226,57 @@ UTypedElementDatabaseCompatibility::FRegistrationCommandChange::FRegistrationCom
 
 UTypedElementDatabaseCompatibility::FRegistrationCommandChange::~FRegistrationCommandChange()
 {
+	// Does not require any thread locking as IsRowAvailable is thread safe and DestroyMemento will lock.
+	
 	// If there has been no revert operation, there's also no memento.
 	if (UTypedElementDatabaseCompatibility* DataStorageCompat = Owner.Get(); 
 		DataStorageCompat && DataStorageCompat->Storage->IsRowAvailable(MementoRow))
 	{
-		DataStorageCompat->Environment->GetMementoSystem().DestroyMemento(MementoRow);
+		if (TypedElementDataStorage::bUseCommandBuffer)
+		{
+			DataStorageCompat->QueuedCommands.AddCommand(UE::EditorDataStorage::FDestroyMemento{ .MementoRow = MementoRow});
+		}
+		else
+		{
+			DataStorageCompat->Environment->GetMementoSystem().DestroyMemento(MementoRow);
+		}
 	}
 }
 
 void UTypedElementDatabaseCompatibility::FRegistrationCommandChange::Apply(UObject* Object)
 {
 	using namespace TypedElementDataStorage;
+	using namespace UE::EditorDataStorage;
 
 	checkf(Owner.IsValid() && Owner.Get() == Object, 
 		TEXT("Applying registration transaction command within TEDS Compat was called after TEDS is not longer available."));
 	UTypedElementDatabaseCompatibility* DataStorageCompat = Owner.Get();
 	if (UObject* TargetRetrieved = TargetObject.Get(/*bEvenIfPendingKill=*/ true))
 	{
-		RowHandle ObjectRow = DataStorageCompat->AddCompatibleObjectExplicitTransactionable<false>(TargetRetrieved);
-		DataStorageCompat->Environment->GetMementoSystem().RestoreMemento(MementoRow, ObjectRow);
+		if (TypedElementDataStorage::bUseCommandBuffer)
+		{
+			RowHandle ObjectRow = DataStorageCompat->AddCompatibleObjectExplicitTransactionable<false>(TargetRetrieved);
+			DataStorageCompat->QueuedCommands.AddCommand(FRestoreMemento{ .MementoRow = MementoRow, .TargetRow  = ObjectRow });
+		}
+		else
+		{
+			// Lock here because the next two functions would otherwise lock multiple times.
+			FScopedExclusiveLock Lock(EGlobalLockScope::Public);
+			RowHandle ObjectRow = DataStorageCompat->AddCompatibleObjectExplicitTransactionable<false>(TargetRetrieved);
+			DataStorageCompat->Environment->GetMementoSystem().RestoreMemento(MementoRow, ObjectRow);
+		}
 	}
 }
 
 void UTypedElementDatabaseCompatibility::FRegistrationCommandChange::Revert(UObject* Object)
 {
 	using namespace TypedElementDataStorage;
-	
+	using namespace UE::EditorDataStorage;
+
 	checkf(Owner.IsValid() && Owner.Get() == Object,
 		TEXT("Reverting registration transaction command within TEDS Compat was called after TEDS is not longer available."));
 
+	FScopedExclusiveLock Lock(EGlobalLockScope::Public);
 	if (UObject* TargetRetrieved = TargetObject.Get(/*bEvenIfPendingKill=*/ true))
 	{
 		UTypedElementDatabaseCompatibility* DataStorageCompat = Owner.Get(); 
@@ -906,8 +1285,16 @@ void UTypedElementDatabaseCompatibility::FRegistrationCommandChange::Revert(UObj
 		RowHandle ObjectRow = DataStorageCompat->FindRowWithCompatibleObjectExplicit(TargetRetrieved);
 		if (DataStorage->IsRowAvailable(ObjectRow))
 		{
-			MementoRow = DataStorageCompat->Environment->GetMementoSystem().CreateMemento(ObjectRow);
-			DataStorageCompat->RemoveCompatibleObjectExplicitTransactionable<false>(TargetRetrieved);
+			if (TypedElementDataStorage::bUseCommandBuffer && TypedElementDataStorage::bUseDeferredRemovesInCompat)
+			{
+				MementoRow = DataStorage->ReserveRow();
+				DataStorageCompat->QueuedCommands.AddCommand(FCreateMemento{ .ReservedMementoRow = MementoRow, .TargetRow = ObjectRow });
+			}
+			else
+			{
+				MementoRow = DataStorageCompat->Environment->GetMementoSystem().CreateMemento(ObjectRow);
+			}
+			DataStorageCompat->RemoveCompatibleObjectExplicitTransactionable<false>(TargetRetrieved, ObjectRow);
 		}
 	}
 }
@@ -928,28 +1315,48 @@ UTypedElementDatabaseCompatibility::FDeregistrationCommandChange::FDeregistratio
 	, TargetObject(InTargetObject)
 {
 	using namespace TypedElementDataStorage;
+	using namespace UE::EditorDataStorage;
 
 	ITypedElementDataStorageInterface* DataStorage = InOwner->Storage;
 
 	RowHandle ObjectRow = InOwner->FindRowWithCompatibleObjectExplicit(InTargetObject);
 	if (DataStorage->IsRowAvailable(ObjectRow))
 	{
-		MementoRow = InOwner->Environment->GetMementoSystem().CreateMemento(ObjectRow);
+		if (TypedElementDataStorage::bUseCommandBuffer && TypedElementDataStorage::bUseDeferredRemovesInCompat)
+		{
+			MementoRow = DataStorage->ReserveRow();
+			InOwner->QueuedCommands.AddCommand(FCreateMemento{ .ReservedMementoRow = MementoRow, .TargetRow = ObjectRow });
+		}
+		else
+		{
+			MementoRow = InOwner->Environment->GetMementoSystem().CreateMemento(ObjectRow);
+		}
 	}
 }
 
 UTypedElementDatabaseCompatibility::FDeregistrationCommandChange::~FDeregistrationCommandChange()
 {
+	using namespace UE::EditorDataStorage;
+
 	// There's no memento row if target object was never registered with TEDS Compat.
 	if (UTypedElementDatabaseCompatibility* DataStorageCompat = Owner.Get();
 		DataStorageCompat && DataStorageCompat->Storage->IsRowAvailable(MementoRow))
 	{
-		DataStorageCompat->Environment->GetMementoSystem().DestroyMemento(MementoRow);
+		if (TypedElementDataStorage::bUseCommandBuffer)
+		{
+			DataStorageCompat->QueuedCommands.AddCommand(FDestroyMemento{ .MementoRow = MementoRow });
+		}
+		else
+		{
+			DataStorageCompat->Environment->GetMementoSystem().DestroyMemento(MementoRow);
+		}
 	}
 }
 
 void UTypedElementDatabaseCompatibility::FDeregistrationCommandChange::Apply(UObject* Object)
 {
+	// All function calls are guaranteed to be thread safe.
+
 	using namespace TypedElementDataStorage;
 
 	checkf(Owner.IsValid() && Owner.Get() == Object,
@@ -964,6 +1371,7 @@ void UTypedElementDatabaseCompatibility::FDeregistrationCommandChange::Apply(UOb
 void UTypedElementDatabaseCompatibility::FDeregistrationCommandChange::Revert(UObject* Object)
 {
 	using namespace TypedElementDataStorage;
+	using namespace UE::EditorDataStorage;
 
 	checkf(Owner.IsValid() && Owner.Get() == Object,
 		TEXT("Reverting deregistration transaction command within TEDS Compat was called after TEDS is not longer available."));
@@ -971,8 +1379,18 @@ void UTypedElementDatabaseCompatibility::FDeregistrationCommandChange::Revert(UO
 	UTypedElementDatabaseCompatibility* DataStorageCompat = Owner.Get();
 	if (UObject* TargetRetrieved = TargetObject.Get(/*bEvenIfPendingKill=*/ true))
 	{
-		RowHandle ObjectRow = DataStorageCompat->AddCompatibleObjectExplicitTransactionable<false>(TargetRetrieved);
-		DataStorageCompat->Environment->GetMementoSystem().RestoreMemento(MementoRow, ObjectRow);
+		if(TypedElementDataStorage::bUseCommandBuffer)
+		{
+			RowHandle ObjectRow = DataStorageCompat->AddCompatibleObjectExplicitTransactionable<false>(TargetRetrieved);
+			DataStorageCompat->QueuedCommands.AddCommand(FRestoreMemento{ .MementoRow = MementoRow, .TargetRow = ObjectRow });
+		}
+		else
+		{
+			// Lock here because the next two functions would otherwise lock multiple times.
+			FScopedExclusiveLock Lock(EGlobalLockScope::Public);
+			RowHandle ObjectRow = DataStorageCompat->AddCompatibleObjectExplicitTransactionable<false>(TargetRetrieved);
+			DataStorageCompat->Environment->GetMementoSystem().RestoreMemento(MementoRow, ObjectRow);
+		}
 	}
 }
 
diff --git a/Engine/Plugins/Experimental/TypedElementsDataStorage/Source/TypedElementsDataStorage/Private/TypedElementDatabaseCompatibility.h b/Engine/Plugins/Experimental/TypedElementsDataStorage/Source/TypedElementsDataStorage/Private/TypedElementDatabaseCompatibility.h
--- a/Engine/Plugins/Experimental/TypedElementsDataStorage/Source/TypedElementsDataStorage/Private/TypedElementDatabaseCompatibility.h
+++ b/Engine/Plugins/Experimental/TypedElementsDataStorage/Source/TypedElementsDataStorage/Private/TypedElementDatabaseCompatibility.h
@@ -6,6 +6,8 @@
 #include "Async/Mutex.h"
 #include "Containers/Array.h"
 #include "Containers/Map.h"
+#include "Commands/EditorDataStorageCommandBuffer.h"
+#include "Commands/EditorDataStorageCompatibilityCommands.h"
 #include "Compatibility/TypedElementObjectReinstancingManager.h"
 #include "Elements/Interfaces/TypedElementDataStorageInterface.h"
 #include "Elements/Interfaces/TypedElementDataStorageCompatibilityInterface.h"
@@ -23,7 +25,6 @@ class FTypedElementDatabaseEnvironment;
 class UTypedElementMementoSystem;
 
 enum class ETypedElementDatabaseCompatibilityObjectType : uint8;
-struct FTypedElementDatabaseCompatibilityObjectTypeInfo;
 
 UCLASS()
 class TYPEDELEMENTSDATASTORAGE_API UTypedElementDatabaseCompatibility
@@ -31,10 +32,11 @@ class TYPEDELEMENTSDATASTORAGE_API UTypedElementDatabaseCompatibility
 	, public ITypedElementDataStorageCompatibilityInterface
 {
 	GENERATED_BODY()
+
+	friend struct UE::EditorDataStorage::FCommandProcessor;
+	friend struct UE::EditorDataStorage::FPatchData;
+	friend struct UE::EditorDataStorage::FPrepareCommands;
 public:
-	using ObjectAddedCallback = TFunction<void(const void* /*Object*/, const FTypedElementDatabaseCompatibilityObjectTypeInfo&, TypedElementRowHandle /*Row*/)>;
-	using ObjectRemovedCallback = TFunction<void(const void* /*Object*/, const FTypedElementDatabaseCompatibilityObjectTypeInfo&, TypedElementRowHandle /*Row*/)>;
-	
 	~UTypedElementDatabaseCompatibility() override = default;
 
 	void Initialize(UTypedElementDatabase* InStorage);
@@ -43,13 +45,13 @@ public:
 	void RegisterRegistrationFilter(ObjectRegistrationFilter Filter) override;
 	void RegisterDealiaserCallback(ObjectToRowDealiaser Dealiaser) override;
 	void RegisterTypeTableAssociation(TObjectPtr<UStruct> TypeInfo, TypedElementDataStorage::TableHandle Table) override;
-	FDelegateHandle RegisterObjectAddedCallback(ObjectAddedCallback&& OnObjectAdded);
+	FDelegateHandle RegisterObjectAddedCallback(UE::EditorDataStorage::ObjectAddedCallback&& OnObjectAdded);
 	void UnregisterObjectAddedCallback(FDelegateHandle Handle);
-	FDelegateHandle RegisterObjectRemovedCallback(ObjectRemovedCallback&& OnObjectRemoved);
+	FDelegateHandle RegisterObjectRemovedCallback(UE::EditorDataStorage::ObjectRemovedCallback&& OnObjectRemoved);
 	void UnregisterObjectRemovedCallback(FDelegateHandle Handle);
 	
 	TypedElementRowHandle AddCompatibleObjectExplicit(UObject* Object) override;
-	TypedElementRowHandle AddCompatibleObjectExplicit(void* Object, TWeakObjectPtr<const UScriptStruct> TypeInfo) override;
+	TypedElementRowHandle AddCompatibleObjectExplicit(void* Object, TWeakObjectPtr<UScriptStruct> TypeInfo) override;
 	
 	void RemoveCompatibleObjectExplicit(UObject* Object) override;
 	void RemoveCompatibleObjectExplicit(void* Object) override;
@@ -61,7 +63,6 @@ public:
 	void ListExtensions(TFunctionRef<void(FName)> Callback) const override;
 	
 private:
-
 	// The below changes expect UTypedElementDatabaseCompatibility to be the object passed in to StoreUndo.
 	// Note that we cannot pass in TargetObject to StoreUndo because doing so seems to stomp regular Modify()
 	// changes for that object.
@@ -112,6 +113,7 @@ private:
 	TypedElementRowHandle DealiasObject(const UObject* Object) const;
 
 	void Tick();
+	void TickPendingCommands();
 	void TickPendingUObjectRegistration();
 	void TickPendingExternalObjectRegistration();
 	void TickObjectSync();
@@ -119,8 +121,8 @@ private:
 	void OnPrePropertyChanged(UObject* Object, const FEditPropertyChain& PropertyChain);
 	void OnPostEditChangeProperty(UObject* Object, FPropertyChangedEvent& PropertyChangedEvent);
 	void OnObjectModified(UObject* Object);
-	void OnObjectAdded(const void* Object, FTypedElementDatabaseCompatibilityObjectTypeInfo TypeInfo, TypedElementRowHandle Row) const;
-	void OnPreObjectRemoved(const void* Object, FTypedElementDatabaseCompatibilityObjectTypeInfo TypeInfo, TypedElementRowHandle Row) const;
+	void TriggerOnObjectAdded(const void* Object, UE::EditorDataStorage::FObjectTypeInfo TypeInfo, TypedElementDataStorage::RowHandle Row) const;
+	void TriggerOnPreObjectRemoved(const void* Object, UE::EditorDataStorage::FObjectTypeInfo TypeInfo, TypedElementDataStorage::RowHandle Row) const;
 	void OnObjectReinstanced(const FCoreUObjectDelegates::FReplacementObjectMap& ReplacedObjects);
 
 	void OnPostGcUnreachableAnalysis();
@@ -139,7 +141,7 @@ private:
 		void Process(UTypedElementDatabaseCompatibility& Compatibility);
 
 	private:
-		TOptional<TWeakObjectPtr<UObject>> ProcessResolveTypeRecursively(const TWeakObjectPtr<const UObject>& Target);
+		TOptional<TWeakObjectPtr<UObject>> ProcessResolveTypeRecursively(const TWeakObjectPtr<UObject>& Target);
 
 		struct FTypeInfoEntryKeyFuncs : TDefaultMapHashableKeyFuncs<TWeakObjectPtr<UObject>, TWeakObjectPtr<UObject>, false>
 		{
@@ -159,7 +161,7 @@ private:
 	struct ExternalObjectRegistration
 	{
 		void* Object;
-		TWeakObjectPtr<const UScriptStruct> TypeInfo;
+		TWeakObjectPtr<UScriptStruct> TypeInfo;
 	};
 	
 	template<typename AddressType>
@@ -185,6 +187,8 @@ private:
 		void Reset();
 	};
 
+	UE::EditorDataStorage::CompatibilityCommandBuffer QueuedCommands;
+	UE::EditorDataStorage::CompatibilityCommandBuffer::FCollection PendingCommands;
 	PendingRegistration<TWeakObjectPtr<UObject>> UObjectsPendingRegistration;
 	PendingRegistration<ExternalObjectRegistration> ExternalObjectsPendingRegistration;
 	TArray<TypedElementDataStorage::RowHandle> RowScratchBuffer;
@@ -193,13 +197,13 @@ private:
 	TArray<ObjectToRowDealiaser> ObjectToRowDialiasers;
 	using TypeToTableMapType = TMap<TWeakObjectPtr<UStruct>, TypedElementDataStorage::TableHandle>;
 	TypeToTableMapType TypeToTableMap;
-	TArray<TPair<ObjectAddedCallback, FDelegateHandle>> ObjectAddedCallbackList;
-	TArray<TPair<ObjectRemovedCallback, FDelegateHandle>> PreObjectRemovedCallbackList;
+	TArray<TPair<UE::EditorDataStorage::ObjectAddedCallback, FDelegateHandle>> ObjectAddedCallbackList;
+	TArray<TPair<UE::EditorDataStorage::ObjectRemovedCallback, FDelegateHandle>> PreObjectRemovedCallbackList;
 
-	TypedElementTableHandle StandardActorTable{ TypedElementInvalidTableHandle };
-	TypedElementTableHandle StandardActorWithTransformTable{ TypedElementInvalidTableHandle };
-	TypedElementTableHandle StandardUObjectTable{ TypedElementInvalidTableHandle };
-	TypedElementTableHandle StandardExternalObjectTable{ TypedElementInvalidTableHandle };
+	TypedElementDataStorage::TableHandle StandardActorTable{ TypedElementDataStorage::InvalidTableHandle };
+	TypedElementDataStorage::TableHandle StandardActorWithTransformTable{ TypedElementDataStorage::InvalidTableHandle };
+	TypedElementDataStorage::TableHandle StandardUObjectTable{ TypedElementDataStorage::InvalidTableHandle };
+	TypedElementDataStorage::TableHandle StandardExternalObjectTable{ TypedElementDataStorage::InvalidTableHandle };
 	ITypedElementDataStorageInterface* Storage{ nullptr };
 
 	/**
@@ -208,7 +212,7 @@ private:
 	 */
 	struct FSyncTagInfo
 	{
-		TWeakObjectPtr<const UScriptStruct> ColumnType;
+		TWeakObjectPtr<UScriptStruct> ColumnType;
 		bool bAddColumn;
 
 		bool operator==(const FSyncTagInfo& Rhs) const = default;
@@ -237,47 +241,3 @@ private:
 };
 
 SIZE_T GetTypeHash(const UTypedElementDatabaseCompatibility::FSyncTagInfo& Column);
-
-enum class ETypedElementDatabaseCompatibilityObjectType : uint8
-{
-	Struct,
-	Class
-};
-
-/**
- * Objects with type info defined in either UScriptStruct or UClass can be stored into TEDS via the
- * ITypedElementDataStorageCompatibilityInterface
- * This is a discriminated union which aids with callbacks made when objects are added
- */
-struct FTypedElementDatabaseCompatibilityObjectTypeInfo
-{
-	ETypedElementDatabaseCompatibilityObjectType TypeInfoType;
-
-	union
-	{
-		const UScriptStruct* ScriptStruct;
-		const UClass* Class;
-	};
-
-	FTypedElementDatabaseCompatibilityObjectTypeInfo(const UScriptStruct* InScriptStruct)
-		: TypeInfoType(ETypedElementDatabaseCompatibilityObjectType::Struct)
-		, ScriptStruct(InScriptStruct)
-	{}
-
-	FTypedElementDatabaseCompatibilityObjectTypeInfo(const UClass* InClass)
-	: TypeInfoType(ETypedElementDatabaseCompatibilityObjectType::Class)
-	, Class(InClass)
-	{}
-
-	FName GetFName() const;
-};
-
-inline FName FTypedElementDatabaseCompatibilityObjectTypeInfo::GetFName() const
-{
-	switch(TypeInfoType)
-	{
-	case ETypedElementDatabaseCompatibilityObjectType::Struct: return ScriptStruct->GetFName();
-	case ETypedElementDatabaseCompatibilityObjectType::Class: return Class->GetFName();
-	default: return FName();
-	}
-}
\ No newline at end of file
diff --git a/Engine/Source/Runtime/TypedElementFramework/Public/Elements/Common/TypedElementHandles.h b/Engine/Source/Runtime/TypedElementFramework/Public/Elements/Common/TypedElementHandles.h
--- a/Engine/Source/Runtime/TypedElementFramework/Public/Elements/Common/TypedElementHandles.h
+++ b/Engine/Source/Runtime/TypedElementFramework/Public/Elements/Common/TypedElementHandles.h
@@ -10,7 +10,7 @@ namespace TypedElementDataStorage
 	static constexpr TableHandle InvalidTableHandle = TNumericLimits<TableHandle>::Max();
 
 	using RowHandle = uint64;
-	static constexpr RowHandle InvalidRowHandle = TNumericLimits<RowHandle>::Max();
+	static constexpr RowHandle InvalidRowHandle = 0;
 
 	using QueryHandle = uint64;
 	static constexpr QueryHandle InvalidQueryHandle = TNumericLimits<QueryHandle>::Max();
diff --git a/Engine/Source/Runtime/TypedElementFramework/Public/Elements/Interfaces/TypedElementDataStorageCompatibilityInterface.h b/Engine/Source/Runtime/TypedElementFramework/Public/Elements/Interfaces/TypedElementDataStorageCompatibilityInterface.h
--- a/Engine/Source/Runtime/TypedElementFramework/Public/Elements/Interfaces/TypedElementDataStorageCompatibilityInterface.h
+++ b/Engine/Source/Runtime/TypedElementFramework/Public/Elements/Interfaces/TypedElementDataStorageCompatibilityInterface.h
@@ -13,8 +13,6 @@
 
 #include "TypedElementDataStorageCompatibilityInterface.generated.h"
 
-struct FTypedElementDatabaseCompatibilityObjectTypeInfo;
-
 UINTERFACE(MinimalAPI)
 class UTypedElementDataStorageCompatibilityInterface : public UInterface
 {
@@ -88,7 +86,7 @@ public:
 	/** Adds a UObject to the data storage. */
 	virtual TypedElementRowHandle AddCompatibleObjectExplicit(UObject* Object) = 0;
 	/** Adds an FStruct to the data storage. */
-	virtual TypedElementRowHandle AddCompatibleObjectExplicit(void* Object, TWeakObjectPtr<const UScriptStruct> TypeInfo) = 0;
+	virtual TypedElementRowHandle AddCompatibleObjectExplicit(void* Object, TWeakObjectPtr<UScriptStruct> TypeInfo) = 0;
 	
 	/** Removes a UObject from the data storage. */
 	virtual void RemoveCompatibleObjectExplicit(UObject* Object) = 0;
