Fixed a static initialization issue.

#rnx

diff --git a/Engine/Plugins/Experimental/TypedElementsDataStorage/Source/TypedElementsDataStorage/Private/Commands/EditorDataStorageCompatibilityCommands.cpp b/Engine/Plugins/Experimental/TypedElementsDataStorage/Source/TypedElementsDataStorage/Private/Commands/EditorDataStorageCompatibilityCommands.cpp
--- a/Engine/Plugins/Experimental/TypedElementsDataStorage/Source/TypedElementsDataStorage/Private/Commands/EditorDataStorageCompatibilityCommands.cpp
+++ b/Engine/Plugins/Experimental/TypedElementsDataStorage/Source/TypedElementsDataStorage/Private/Commands/EditorDataStorageCompatibilityCommands.cpp
@@ -16,8 +16,27 @@
 
 namespace UE::EditorDataStorage
 {
-	UScriptStruct* FAddSyncFromWorldTag::Type = FTypedElementSyncFromWorldTag::StaticStruct();
-	UScriptStruct* FAddInteractiveSyncFromWorldTag::Type = FTypedElementSyncFromWorldInteractiveTag::StaticStruct();
+	UScriptStruct* FAddSyncFromWorldTag::GetType()
+	{
+		return *GetTypeAddress();
+	}
+
+	UScriptStruct** FAddSyncFromWorldTag::GetTypeAddress()
+	{
+		static UScriptStruct* Type = FTypedElementSyncFromWorldTag::StaticStruct();
+		return &Type;
+	}
+
+	UScriptStruct* FAddInteractiveSyncFromWorldTag::GetType()
+	{
+		return *GetTypeAddress();
+	}
+
+	UScriptStruct** FAddInteractiveSyncFromWorldTag::GetTypeAddress()
+	{
+		static UScriptStruct* Type = FTypedElementSyncFromWorldInteractiveTag::StaticStruct();
+		return &Type;
+	}
 
 	//
 	// FObjectTypeInfo
@@ -188,8 +207,8 @@ namespace UE::EditorDataStorage
 		// Patch locally cached type information
 		UScriptStruct** LocalTypeInfoStorage[] =
 		{ 
-			&FAddSyncFromWorldTag::Type,
-			&FAddInteractiveSyncFromWorldTag::Type
+			FAddSyncFromWorldTag::GetTypeAddress(),
+			FAddInteractiveSyncFromWorldTag::GetTypeAddress()
 		};
 		TArrayView<UScriptStruct**> LocalTypeInfo(LocalTypeInfoStorage, sizeof(LocalTypeInfoStorage) / sizeof(UScriptStruct**));
 		for (UScriptStruct** TypeInfo : LocalTypeInfo)
@@ -657,19 +676,19 @@ namespace UE::EditorDataStorage
 
 	void FCommandProcessor::operator()(FAddInteractiveSyncFromWorldTag& Command)
 	{
-		Storage.AddColumn(Command.Row, FAddInteractiveSyncFromWorldTag::Type);
+		Storage.AddColumn(Command.Row, FAddInteractiveSyncFromWorldTag::GetType());
 	}
 
 	void FCommandProcessor::operator()(FRemoveInteractiveSyncFromWorldTag& Command)
 	{
 		Storage.AddRemoveColumns(Command.Row,
-			{ FAddSyncFromWorldTag::Type },
-			{ FAddInteractiveSyncFromWorldTag::Type });
+			{ FAddSyncFromWorldTag::GetType()},
+			{ FAddInteractiveSyncFromWorldTag::GetType() });
 	}
 
 	void FCommandProcessor::operator()(FAddSyncFromWorldTag& Command)
 	{
-		Storage.AddColumn(Command.Row, FAddSyncFromWorldTag::Type);
+		Storage.AddColumn(Command.Row, FAddSyncFromWorldTag::GetType());
 	}
 
 
diff --git a/Engine/Plugins/Experimental/TypedElementsDataStorage/Source/TypedElementsDataStorage/Private/Commands/EditorDataStorageCompatibilityCommands.h b/Engine/Plugins/Experimental/TypedElementsDataStorage/Source/TypedElementsDataStorage/Private/Commands/EditorDataStorageCompatibilityCommands.h
--- a/Engine/Plugins/Experimental/TypedElementsDataStorage/Source/TypedElementsDataStorage/Private/Commands/EditorDataStorageCompatibilityCommands.h
+++ b/Engine/Plugins/Experimental/TypedElementsDataStorage/Source/TypedElementsDataStorage/Private/Commands/EditorDataStorageCompatibilityCommands.h
@@ -152,18 +152,20 @@ namespace UE::EditorDataStorage
 
 	struct FAddSyncFromWorldTag final
 	{
-		static UScriptStruct* Type;
-
 		TObjectKey<const UObject> Target;
 		TypedElementDataStorage::RowHandle Row;
+
+		static UScriptStruct* GetType();
+		static UScriptStruct** GetTypeAddress();
 	};
 
 	struct FAddInteractiveSyncFromWorldTag final
 	{
-		static UScriptStruct* Type;
-
 		TObjectKey<const UObject> Target;
 		TypedElementDataStorage::RowHandle Row;
+
+		static UScriptStruct* GetType();
+		static UScriptStruct** GetTypeAddress();
 	};
 
 	struct FRemoveInteractiveSyncFromWorldTag final
